- Renamed the RenderPages plugin to RenderGrid
- Renamed all the classes, functions, fields and variables, etc to match the name change of the plugin
- Renamed the viewer modes, removed the "None" viewer mode
- Added a black background behind the viewer modes
- Added helpful messages to the viewer modes when rendering isn't possible (to let the user know why it isn't possible)

diff --git a/Engine/Plugins/Experimental/RenderGrid/Source/RenderGrid/Private/Factories/RenderGridPropsSourceFactoryLocal.cpp b/Engine/Plugins/Experimental/RenderGrid/Source/RenderGrid/Private/Factories/RenderGridPropsSourceFactoryLocal.cpp
new file mode 100644
--- /dev/null
+++ b/Engine/Plugins/Experimental/RenderGrid/Source/RenderGrid/Private/Factories/RenderGridPropsSourceFactoryLocal.cpp
@@ -0,0 +1,12 @@
+// Copyright Epic Games, Inc. All Rights Reserved.
+
+#include "Factories/RenderGridPropsSourceFactoryLocal.h"
+#include "RenderGrid/RenderGridPropsSource.h"
+
+
+URenderGridPropsSourceBase* UE::RenderGrid::Private::FRenderGridPropsSourceFactoryLocal::CreateInstance(UObject* Outer, UObject* PropsSourceOrigin)
+{
+	URenderGridPropsSourceLocal* PropsSource = NewObject<URenderGridPropsSourceLocal>(Outer);
+	PropsSource->SetSourceOrigin(PropsSourceOrigin);
+	return PropsSource;
+}
diff --git a/Engine/Plugins/Experimental/RenderGrid/Source/RenderGrid/Private/Factories/RenderGridPropsSourceFactoryRemoteControl.cpp b/Engine/Plugins/Experimental/RenderGrid/Source/RenderGrid/Private/Factories/RenderGridPropsSourceFactoryRemoteControl.cpp
new file mode 100644
--- /dev/null
+++ b/Engine/Plugins/Experimental/RenderGrid/Source/RenderGrid/Private/Factories/RenderGridPropsSourceFactoryRemoteControl.cpp
@@ -0,0 +1,12 @@
+// Copyright Epic Games, Inc. All Rights Reserved.
+
+#include "Factories/RenderGridPropsSourceFactoryRemoteControl.h"
+#include "RenderGrid/RenderGridPropsSource.h"
+
+
+URenderGridPropsSourceBase* UE::RenderGrid::Private::FRenderGridPropsSourceFactoryRemoteControl::CreateInstance(UObject* Outer, UObject* PropsSourceOrigin)
+{
+	URenderGridPropsSourceRemoteControl* PropsSource = NewObject<URenderGridPropsSourceRemoteControl>(Outer);
+	PropsSource->SetSourceOrigin(PropsSourceOrigin);
+	return PropsSource;
+}
diff --git a/Engine/Plugins/Experimental/RenderGrid/Source/RenderGrid/Private/RenderGrid/RenderGrid.cpp b/Engine/Plugins/Experimental/RenderGrid/Source/RenderGrid/Private/RenderGrid/RenderGrid.cpp
new file mode 100644
--- /dev/null
+++ b/Engine/Plugins/Experimental/RenderGrid/Source/RenderGrid/Private/RenderGrid/RenderGrid.cpp
@@ -0,0 +1,823 @@
+// Copyright Epic Games, Inc. All Rights Reserved.
+
+#include "RenderGrid/RenderGrid.h"
+#include "RenderGrid/RenderGridManager.h"
+#include "RenderGridUtils.h"
+#include "IRenderGridModule.h"
+#include "LevelSequence.h"
+#include "MoviePipelineMasterConfig.h"
+#include "MoviePipelineOutputSetting.h"
+#include "MovieScene.h"
+#include "UObject/ObjectSaveContext.h"
+
+
+URenderGridJob::URenderGridJob()
+	: Id(FGuid::NewGuid())
+	, WaitFramesBeforeRendering(0)
+	, Sequence(nullptr)
+	, bOverrideStartFrame(false)
+	, CustomStartFrame(0)
+	, bOverrideEndFrame(false)
+	, CustomEndFrame(0)
+	, bOverrideResolution(false)
+	, CustomResolution(FIntPoint(3840, 2160))
+	, bIsEnabled(true)
+	, RenderPreset(nullptr)
+{}
+
+TOptional<int32> URenderGridJob::GetSequenceStartFrame() const
+{
+	if (!IsValid(Sequence))
+	{
+		return TOptional<int32>();
+	}
+
+	const UMovieScene* MovieScene = Sequence->MovieScene;
+	const FFrameRate TickResolution = MovieScene->GetTickResolution();
+	const FFrameRate DisplayRate = MovieScene->GetDisplayRate();
+	const UMoviePipelineOutputSetting* Settings = GetRenderPresetOutputSettings();
+
+	if (bOverrideStartFrame || (IsValid(Settings) && Settings->bUseCustomPlaybackRange))
+	{
+		int32 Frame = (bOverrideStartFrame ? CustomStartFrame : Settings->CustomStartFrame);
+		if (IsValid(Settings) && Settings->bUseCustomFrameRate)
+		{
+			return FMath::FloorToInt(Frame / (Settings->OutputFrameRate / DisplayRate).AsDecimal());
+		}
+		return Frame;
+	}
+
+	const int32 StartFrameNumber = MovieScene->GetPlaybackRange().GetLowerBoundValue().Value;
+	return FMath::FloorToInt(StartFrameNumber / (TickResolution / DisplayRate).AsDecimal());
+}
+
+TOptional<int32> URenderGridJob::GetSequenceEndFrame() const
+{
+	if (!IsValid(Sequence))
+	{
+		return TOptional<int32>();
+	}
+
+	const UMovieScene* MovieScene = Sequence->MovieScene;
+	const FFrameRate TickResolution = MovieScene->GetTickResolution();
+	const FFrameRate DisplayRate = MovieScene->GetDisplayRate();
+	const UMoviePipelineOutputSetting* Settings = GetRenderPresetOutputSettings();
+
+	if (bOverrideEndFrame || (IsValid(Settings) && Settings->bUseCustomPlaybackRange))
+	{
+		int32 Frame = (bOverrideEndFrame ? CustomEndFrame : Settings->CustomEndFrame);
+		if (IsValid(Settings) && Settings->bUseCustomFrameRate)
+		{
+			return FMath::FloorToInt(Frame / (Settings->OutputFrameRate / DisplayRate).AsDecimal());
+		}
+		return Frame;
+	}
+
+	const int32 EndFrameNumber = MovieScene->GetPlaybackRange().GetUpperBoundValue().Value;
+	return FMath::FloorToInt(EndFrameNumber / (TickResolution / DisplayRate).AsDecimal());
+}
+
+bool URenderGridJob::SetSequenceStartFrame(const int32 NewCustomStartFrame)
+{
+	int32 StartFrame = GetStartFrame().Get(0);
+	SetIsUsingCustomStartFrame(true);
+	SetCustomStartFrame(StartFrame);
+
+	TOptional<int32> SequenceStartFrame = GetSequenceStartFrame();
+	while (SequenceStartFrame.IsSet() && (SequenceStartFrame.Get(0) > NewCustomStartFrame))
+	{
+		StartFrame--;
+		if (StartFrame <= INT32_MIN)
+		{
+			return false;
+		}
+		SetCustomStartFrame(StartFrame);
+		SequenceStartFrame = GetSequenceStartFrame();
+	}
+	while (SequenceStartFrame.IsSet() && (SequenceStartFrame.Get(0) < NewCustomStartFrame))
+	{
+		StartFrame++;
+		if (StartFrame >= INT32_MAX)
+		{
+			return false;
+		}
+		SetCustomStartFrame(StartFrame);
+		SequenceStartFrame = GetSequenceStartFrame();
+	}
+	return (SequenceStartFrame.IsSet() && (SequenceStartFrame.Get(0) == NewCustomStartFrame));
+}
+
+bool URenderGridJob::SetSequenceEndFrame(const int32 NewCustomStartFrame)
+{
+	int32 EndFrame = GetEndFrame().Get(0);
+	SetIsUsingCustomEndFrame(true);
+	SetCustomEndFrame(EndFrame);
+
+	TOptional<int32> SequenceEndFrame = GetSequenceEndFrame();
+	while (SequenceEndFrame.IsSet() && (SequenceEndFrame.Get(0) > NewCustomStartFrame))
+	{
+		EndFrame--;
+		if (EndFrame <= INT32_MIN)
+		{
+			return false;
+		}
+		SetCustomEndFrame(EndFrame);
+		SequenceEndFrame = GetSequenceEndFrame();
+	}
+	while (SequenceEndFrame.IsSet() && (SequenceEndFrame.Get(0) < NewCustomStartFrame))
+	{
+		EndFrame++;
+		if (EndFrame >= INT32_MAX)
+		{
+			return false;
+		}
+		SetCustomEndFrame(EndFrame);
+		SequenceEndFrame = GetSequenceEndFrame();
+	}
+	return (SequenceEndFrame.IsSet() && (SequenceEndFrame.Get(0) == NewCustomStartFrame));
+}
+
+TOptional<int32> URenderGridJob::GetStartFrame() const
+{
+	if (bOverrideStartFrame)
+	{
+		return CustomStartFrame;
+	}
+
+	const UMoviePipelineOutputSetting* Settings = GetRenderPresetOutputSettings();
+	if (IsValid(Settings) && Settings->bUseCustomPlaybackRange)
+	{
+		return Settings->CustomStartFrame;
+	}
+
+	if (!IsValid(Sequence))
+	{
+		return TOptional<int32>();
+	}
+
+	const UMovieScene* MovieScene = Sequence->MovieScene;
+	const FFrameRate TickResolution = MovieScene->GetTickResolution();
+	FFrameRate DisplayRate = MovieScene->GetDisplayRate();
+	if (IsValid(Settings) && Settings->bUseCustomFrameRate)
+	{
+		DisplayRate = Settings->OutputFrameRate;
+	}
+	const int32 StartFrameNumber = MovieScene->GetPlaybackRange().GetLowerBoundValue().Value;
+	return FMath::FloorToInt(StartFrameNumber / (TickResolution / DisplayRate).AsDecimal());
+}
+
+TOptional<int32> URenderGridJob::GetEndFrame() const
+{
+	if (bOverrideEndFrame)
+	{
+		return CustomEndFrame;
+	}
+
+	const UMoviePipelineOutputSetting* Settings = GetRenderPresetOutputSettings();
+	if (IsValid(Settings) && Settings->bUseCustomPlaybackRange)
+	{
+		return Settings->CustomEndFrame;
+	}
+
+	if (!IsValid(Sequence))
+	{
+		return TOptional<int32>();
+	}
+
+	const UMovieScene* MovieScene = Sequence->MovieScene;
+	const FFrameRate TickResolution = MovieScene->GetTickResolution();
+	FFrameRate DisplayRate = MovieScene->GetDisplayRate();
+	if (IsValid(Settings) && Settings->bUseCustomFrameRate)
+	{
+		DisplayRate = Settings->OutputFrameRate;
+	}
+	const int32 EndFrameNumber = MovieScene->GetPlaybackRange().GetUpperBoundValue().Value;
+	return FMath::FloorToInt(EndFrameNumber / (TickResolution / DisplayRate).AsDecimal());
+}
+
+TOptional<double> URenderGridJob::GetStartTime() const
+{
+	if (!IsValid(Sequence))
+	{
+		return TOptional<double>();
+	}
+
+	const TOptional<int32> StartFrame = GetStartFrame();
+	if (!StartFrame.IsSet())
+	{
+		return TOptional<double>();
+	}
+
+	const UMoviePipelineOutputSetting* Settings = GetRenderPresetOutputSettings();
+
+	const UMovieScene* MovieScene = Sequence->MovieScene;
+	FFrameRate DisplayRate = MovieScene->GetDisplayRate();
+	if (IsValid(Settings) && Settings->bUseCustomFrameRate)
+	{
+		DisplayRate = Settings->OutputFrameRate;
+	}
+	return StartFrame.Get(0) / DisplayRate.AsDecimal();
+}
+
+TOptional<double> URenderGridJob::GetEndTime() const
+{
+	if (!IsValid(Sequence))
+	{
+		return TOptional<double>();
+	}
+
+	const TOptional<int32> EndFrame = GetEndFrame();
+	if (!EndFrame.IsSet())
+	{
+		return TOptional<double>();
+	}
+
+	const UMoviePipelineOutputSetting* Settings = GetRenderPresetOutputSettings();
+
+	const UMovieScene* MovieScene = Sequence->MovieScene;
+	FFrameRate DisplayRate = MovieScene->GetDisplayRate();
+	if (IsValid(Settings) && Settings->bUseCustomFrameRate)
+	{
+		DisplayRate = Settings->OutputFrameRate;
+	}
+	return EndFrame.Get(0) / DisplayRate.AsDecimal();
+}
+
+TOptional<double> URenderGridJob::GetDurationInSeconds() const
+{
+	if (!IsValid(Sequence))
+	{
+		return TOptional<double>();
+	}
+
+	const TOptional<int32> StartFrame = GetStartFrame();
+	const TOptional<int32> EndFrame = GetEndFrame();
+	if (!StartFrame.IsSet() || !EndFrame.IsSet() || (StartFrame.Get(0) > EndFrame.Get(0)))
+	{
+		return TOptional<double>();
+	}
+
+	const UMoviePipelineOutputSetting* Settings = GetRenderPresetOutputSettings();
+
+	const UMovieScene* MovieScene = Sequence->MovieScene;
+	FFrameRate DisplayRate = MovieScene->GetDisplayRate();
+	if (IsValid(Settings) && Settings->bUseCustomFrameRate)
+	{
+		DisplayRate = Settings->OutputFrameRate;
+	}
+	return (EndFrame.Get(0) - StartFrame.Get(0)) / DisplayRate.AsDecimal();
+}
+
+double URenderGridJob::GetOutputAspectRatio() const
+{
+	const UMoviePipelineOutputSetting* Settings = GetRenderPresetOutputSettings();
+	if (IsValid(Settings))
+	{
+		return static_cast<double>(Settings->OutputResolution.X) / static_cast<double>(Settings->OutputResolution.Y);
+	}
+
+	const UMoviePipelineOutputSetting* DefaultSettings = GetDefault<UMoviePipelineOutputSetting>();
+	return static_cast<double>(DefaultSettings->OutputResolution.X) / static_cast<double>(DefaultSettings->OutputResolution.Y);
+}
+
+bool URenderGridJob::MatchesSearchTerm(const FString& SearchTerm) const
+{
+	if (SearchTerm.TrimStartAndEnd().Len() <= 0)
+	{
+		return true;
+	}
+	TArray<FString> Parts;
+	SearchTerm.ParseIntoArray(Parts, TEXT(" "), true);
+	for (const FString& Part : Parts)
+	{
+		FString TrimmedPart = Part.TrimStartAndEnd();
+		if (TrimmedPart.Len() <= 0)
+		{
+			continue;
+		}
+		if ((JobId.Find(TrimmedPart) != INDEX_NONE) || (JobName.Find(TrimmedPart) != INDEX_NONE) || (OutputDirectory.Find(TrimmedPart) != INDEX_NONE))
+		{
+			continue;
+		}
+		if (IsValid(RenderPreset) && (RenderPreset.GetPath().Find(TrimmedPart) != INDEX_NONE))
+		{
+			continue;
+		}
+		return false;
+	}
+	return true;
+}
+
+FString URenderGridJob::PurgeJobIdOrReturnEmptyString(const FString& NewJobId)
+{
+	static FString ValidCharacters = TEXT("abcdefghijklmnopqrstuvwxyzABCDEFGHIJKLMNOPQRSTUVWXYZ0123456789-_");
+
+	FString Result;
+	for (TCHAR NewJobIdChar : NewJobId)
+	{
+		int32 Index;
+		if (ValidCharacters.FindChar(NewJobIdChar, Index))
+		{
+			Result += NewJobIdChar;
+		}
+	}
+	return Result;
+}
+
+FString URenderGridJob::PurgeJobId(const FString& NewJobId)
+{
+	FString Result = PurgeJobIdOrReturnEmptyString(NewJobId);
+	if (Result.IsEmpty())
+	{
+		return TEXT("0");
+	}
+	return Result;
+}
+
+FString URenderGridJob::PurgeJobIdOrGenerateUniqueId(URenderGrid* Grid, const FString& NewJobId)
+{
+	FString Result = PurgeJobIdOrReturnEmptyString(NewJobId);
+	if (Result.IsEmpty())
+	{
+		return Grid->GenerateNextJobId();
+	}
+	return Result;
+}
+
+FString URenderGridJob::PurgeJobName(const FString& NewJobName)
+{
+	return NewJobName.TrimStartAndEnd();
+}
+
+FString URenderGridJob::PurgeOutputDirectory(const FString& NewOutputDirectory)
+{
+	return UE::RenderGrid::Private::FRenderGridUtils::NormalizeOutputDirectory(FPaths::ConvertRelativePathToFull(NewOutputDirectory))
+		.Replace(*UE::RenderGrid::Private::FRenderGridUtils::NormalizeOutputDirectory(FPaths::ConvertRelativePathToFull(FPaths::ProjectDir())), TEXT("{project_dir}/"));
+}
+
+FString URenderGridJob::GetOutputDirectory() const
+{
+	return UE::RenderGrid::Private::FRenderGridUtils::NormalizeOutputDirectory(FPaths::ConvertRelativePathToFull(OutputDirectory.Replace(TEXT("{project_dir}"), *FPaths::ProjectDir())));
+}
+
+UMoviePipelineOutputSetting* URenderGridJob::GetRenderPresetOutputSettings() const
+{
+	if (!IsValid(RenderPreset))
+	{
+		return nullptr;
+	}
+	for (UMoviePipelineSetting* Settings : RenderPreset->FindSettingsByClass(UMoviePipelineOutputSetting::StaticClass(), false))
+	{
+		if (!IsValid(Settings))
+		{
+			continue;
+		}
+		if (UMoviePipelineOutputSetting* OutputSettings = Cast<UMoviePipelineOutputSetting>(Settings))
+		{
+			if (!OutputSettings->IsEnabled())
+			{
+				continue;
+			}
+			return OutputSettings;
+		}
+	}
+	return nullptr;
+}
+
+bool URenderGridJob::HasRemoteControlValue(const TSharedPtr<FRemoteControlEntity>& RemoteControlEntity) const
+{
+	if (!RemoteControlEntity.IsValid())
+	{
+		return false;
+	}
+
+	const FString Key = RemoteControlEntity->GetId().ToString();
+	return !!RemoteControlValues.Find(Key);
+}
+
+bool URenderGridJob::ConstGetRemoteControlValue(const TSharedPtr<FRemoteControlEntity>& RemoteControlEntity, TArray<uint8>& OutBinaryArray) const
+{
+	if (!RemoteControlEntity.IsValid())
+	{
+		return false;
+	}
+
+	const FString Key = RemoteControlEntity->GetId().ToString();
+	if (const FRenderGridRemoteControlPropertyData* DataPtr = RemoteControlValues.Find(Key))
+	{
+		OutBinaryArray.Append((*DataPtr).Bytes);
+		return true;
+	}
+	return false;
+}
+
+bool URenderGridJob::GetRemoteControlValue(const TSharedPtr<FRemoteControlEntity>& RemoteControlEntity, TArray<uint8>& OutBinaryArray)
+{
+	if (!RemoteControlEntity.IsValid())
+	{
+		return false;
+	}
+
+	const FString Key = RemoteControlEntity->GetId().ToString();
+	if (FRenderGridRemoteControlPropertyData* DataPtr = RemoteControlValues.Find(Key))
+	{
+		OutBinaryArray.Append((*DataPtr).Bytes);
+		return true;
+	}
+
+	if (!URenderGridPropRemoteControl::GetValueOfEntity(RemoteControlEntity, OutBinaryArray))
+	{
+		return false;
+	}
+	RemoteControlValues.Add(Key, FRenderGridRemoteControlPropertyData(TArray(OutBinaryArray)));
+	return true;
+}
+
+bool URenderGridJob::SetRemoteControlValue(const TSharedPtr<FRemoteControlEntity>& RemoteControlEntity, const TArray<uint8>& BinaryArray)
+{
+	if (!RemoteControlEntity.IsValid())
+	{
+		return false;
+	}
+	const FString Key = RemoteControlEntity->GetId().ToString();
+	RemoteControlValues.Add(Key, FRenderGridRemoteControlPropertyData(TArray(BinaryArray)));
+	return true;
+}
+
+
+URenderGrid::URenderGrid()
+	: Id(FGuid::NewGuid())
+	, PropsSourceType(ERenderGridPropsSourceType::Local)
+	, PropsSourceOrigin_RemoteControl(nullptr)
+	, bExecutingPreRender(false)
+	, bExecutingPostRender(false)
+	, CachedPropsSource(nullptr)
+	, CachedPropsSourceType(ERenderGridPropsSourceType::Local)
+{
+	if (!HasAnyFlags(RF_ClassDefaultObject))
+	{
+		LoadValuesFromCDO();
+		OnPreSaveCDO().AddUObject(this, &URenderGrid::SaveValuesToCDO);
+	}
+}
+
+UWorld* URenderGrid::GetWorld() const
+{
+	if (HasAllFlags(RF_ClassDefaultObject))
+	{
+		// If we are a CDO, we must return nullptr instead of calling Outer->GetWorld() to fool UObject::ImplementsGetWorld.
+		return nullptr;
+	}
+
+	for (const FWorldContext& Context : GEngine->GetWorldContexts())
+	{
+		if (Context.WorldType == EWorldType::PIE)
+		{
+			return Context.World();
+		}
+	}
+
+	if (IsValid(GWorld))
+	{
+		return GWorld;
+	}
+
+	if (UWorld* CachedWorld = CachedWorldWeakPtr.Get(); IsValid(CachedWorld))
+	{
+		return CachedWorld;
+	}
+	UObject* Outer = GetOuter();// Could be a GameInstance, could be World, could also be a WidgetTree, so we're just going to follow the outer chain to find the world we're in.
+	while (Outer)
+	{
+		UWorld* World = Outer->GetWorld();
+		if (IsValid(World))
+		{
+			CachedWorldWeakPtr = World;
+			return World;
+		}
+		Outer = Outer->GetOuter();
+	}
+	return nullptr;
+}
+
+void URenderGrid::PreSave(FObjectPreSaveContext SaveContext)
+{
+	if (HasAnyFlags(RF_ClassDefaultObject))
+	{
+		OnPreSaveCDO().Broadcast();
+	}
+	Super::PreSave(SaveContext);
+	SaveValuesToCDO();
+}
+
+void URenderGrid::PostLoad()
+{
+	Super::PostLoad();
+	LoadValuesFromCDO();
+}
+
+void URenderGrid::PreRender(URenderGridJob* Job)
+{
+	bExecutingPreRender = true;
+	ReceivePreRender(Job);
+	bExecutingPreRender = false;
+}
+
+void URenderGrid::PostRender(URenderGridJob* Job)
+{
+	bExecutingPostRender = true;
+	ReceivePostRender(Job);
+	bExecutingPostRender = false;
+}
+
+void URenderGrid::CopyValuesToOrFromCDO(const bool bToCDO)
+{
+	if (HasAnyFlags(RF_ClassDefaultObject))
+	{
+		return;
+	}
+
+	URenderGrid* CDO = GetCDO();
+	if (!IsValid(CDO))
+	{
+		return;
+	}
+
+	for (FProperty* Property = GetClass()->PropertyLink; Property; Property = Property->PropertyLinkNext)
+	{
+		if (!Property->HasAnyPropertyFlags(CPF_Transient | CPF_DuplicateTransient))// if not [Transient]
+		{
+			void* Data = Property->ContainerPtrToValuePtr<void>(this);
+			void* DataCDO = Property->ContainerPtrToValuePtr<void>(CDO);
+			if (bToCDO)
+			{
+				Property->CopyCompleteValue(DataCDO, Data);
+			}
+			else
+			{
+				Property->CopyCompleteValue(Data, DataCDO);
+			}
+		}
+	}
+
+	if (bToCDO)
+	{
+		CDO->RenderGridJobs = TArray<TObjectPtr<URenderGridJob>>();
+		for (URenderGridJob* Job : RenderGridJobs)
+		{
+			if (!IsValid(Job))
+			{
+				continue;
+			}
+			if (URenderGridJob* JobCopy = DuplicateObject(Job, CDO); IsValid(JobCopy))
+			{
+				CDO->RenderGridJobs.Add(JobCopy);
+			}
+		}
+	}
+}
+
+void URenderGrid::SetPropsSource(ERenderGridPropsSourceType InPropsSourceType, UObject* InPropsSourceOrigin)
+{
+	if (InPropsSourceType == ERenderGridPropsSourceType::RemoteControl)
+	{
+		if (URemoteControlPreset* InPropsSourceOrigin_RemoteControl = Cast<URemoteControlPreset>(InPropsSourceOrigin))
+		{
+			PropsSourceType = InPropsSourceType;
+			PropsSourceOrigin_RemoteControl = InPropsSourceOrigin_RemoteControl;
+			return;
+		}
+	}
+	PropsSourceType = ERenderGridPropsSourceType::Local;
+}
+
+URenderGridPropsSourceBase* URenderGrid::GetPropsSource() const
+{
+	UObject* PropsSourceOrigin = GetPropsSourceOrigin();
+	if (!IsValid(CachedPropsSource) || (CachedPropsSourceType != PropsSourceType) || (CachedPropsSourceOriginWeakPtr.Get() != PropsSourceOrigin))
+	{
+		CachedPropsSourceType = PropsSourceType;
+		CachedPropsSourceOriginWeakPtr = PropsSourceOrigin;
+		CachedPropsSource = UE::RenderGrid::IRenderGridModule::Get().CreatePropsSource(const_cast<URenderGrid*>(this), PropsSourceType, PropsSourceOrigin);
+	}
+	return CachedPropsSource;
+}
+
+UObject* URenderGrid::GetPropsSourceOrigin() const
+{
+	if (PropsSourceType == ERenderGridPropsSourceType::RemoteControl)
+	{
+		return PropsSourceOrigin_RemoteControl;
+	}
+	return nullptr;
+}
+
+void URenderGrid::AddRenderGridJob(URenderGridJob* Job)
+{
+	if (IsValid(Job))
+	{
+		RenderGridJobs.Add(Job);
+	}
+}
+
+void URenderGrid::RemoveRenderGridJob(URenderGridJob* Job)
+{
+	RenderGridJobs.Remove(Job);
+}
+
+void URenderGrid::InsertRenderGridJob(URenderGridJob* Job, int32 Index)
+{
+	if (IsValid(Job))
+	{
+		RenderGridJobs.Insert(Job, Index);
+	}
+}
+
+bool URenderGrid::HasRenderGridJob(URenderGridJob* Job) const
+{
+	return (RenderGridJobs.Find(Job) != INDEX_NONE);
+}
+
+int32 URenderGrid::GetIndexOfRenderGridJob(URenderGridJob* Job) const
+{
+	return RenderGridJobs.Find(Job);
+}
+
+TArray<URenderGridJob*> URenderGrid::GetRenderGridJobs() const
+{
+	TArray<URenderGridJob*> Result;
+	for (URenderGridJob* Job : RenderGridJobs)
+	{
+		if (!IsValid(Job))
+		{
+			continue;
+		}
+		Result.Add(Job);
+	}
+	return Result;
+}
+
+TArray<URenderGridJob*> URenderGrid::GetEnabledRenderGridJobs() const
+{
+	TArray<URenderGridJob*> Result;
+	for (URenderGridJob* Job : RenderGridJobs)
+	{
+		if (!IsValid(Job))
+		{
+			continue;
+		}
+		if (Job->GetIsEnabled())
+		{
+			Result.Add(Job);
+		}
+	}
+	return Result;
+}
+
+TArray<URenderGridJob*> URenderGrid::GetDisabledRenderGridJobs() const
+{
+	TArray<URenderGridJob*> Result;
+	for (URenderGridJob* Job : RenderGridJobs)
+	{
+		if (!IsValid(Job))
+		{
+			continue;
+		}
+		if (!Job->GetIsEnabled())
+		{
+			Result.Add(Job);
+		}
+	}
+	return Result;
+}
+
+void URenderGrid::InsertRenderGridJobBefore(URenderGridJob* Job, URenderGridJob* BeforeJob)
+{
+	if (IsValid(Job))
+	{
+		const TArray<TObjectPtr<URenderGridJob>>::SizeType Index = RenderGridJobs.Find(BeforeJob);
+		if (Index == INDEX_NONE)
+		{
+			RenderGridJobs.Add(Job);
+		}
+		else
+		{
+			RenderGridJobs.Insert(Job, Index);
+		}
+	}
+}
+
+void URenderGrid::InsertRenderGridJobAfter(URenderGridJob* Job, URenderGridJob* AfterJob)
+{
+	if (IsValid(Job))
+	{
+		const TArray<TObjectPtr<URenderGridJob>>::SizeType Index = RenderGridJobs.Find(AfterJob);
+		if (Index == INDEX_NONE)
+		{
+			RenderGridJobs.Add(Job);
+		}
+		else
+		{
+			RenderGridJobs.Insert(Job, Index + 1);
+		}
+	}
+}
+
+FString URenderGrid::GenerateNextJobId()
+{
+	int32 Max = 0;
+	for (URenderGridJob* Job : RenderGridJobs)
+	{
+		if (!IsValid(Job))
+		{
+			continue;
+		}
+		int32 Value = FCString::Atoi(*Job->GetJobId());
+		if (Value > Max)
+		{
+			Max = Value;
+		}
+	}
+	FString Result = FString::FromInt(Max + 1);
+	while (Result.Len() < UE::RenderGrid::FRenderGridManager::GeneratedIdCharacterLength)
+	{
+		Result = TEXT("0") + Result;
+	}
+	return Result;
+}
+
+bool URenderGrid::DoesJobIdExist(const FString& JobId)
+{
+	const FString JobIdToLower = JobId.ToLower();
+	for (URenderGridJob* Job : RenderGridJobs)
+	{
+		if (!IsValid(Job))
+		{
+			continue;
+		}
+		if (JobIdToLower == Job->GetJobId().ToLower())
+		{
+			return true;
+		}
+	}
+	return false;
+}
+
+URenderGridJob* URenderGrid::CreateAndAddNewRenderGridJob()
+{
+	URenderGridJob* Job = NewObject<URenderGridJob>(this);
+	Job->SetJobId(GenerateNextJobId());
+	Job->SetJobName(TEXT("New"));
+	Job->SetOutputDirectory(FPaths::ProjectDir() / TEXT("Saved/MovieRenders/"));
+
+	if (URenderGridPropsSourceRemoteControl* PropsSource = GetPropsSource<URenderGridPropsSourceRemoteControl>())
+	{
+		TArray<uint8> BinaryArray;
+		for (URenderGridPropRemoteControl* Field : PropsSource->GetProps()->GetAllCasted())
+		{
+			if (Field->GetValue(BinaryArray))
+			{
+				Job->SetRemoteControlValue(Field->GetRemoteControlEntity(), BinaryArray);
+			}
+		}
+	}
+
+	AddRenderGridJob(Job);
+	return Job;
+}
+
+URenderGridJob* URenderGrid::DuplicateAndAddRenderGridJob(URenderGridJob* Job)
+{
+	if (!IsValid(Job))
+	{
+		return nullptr;
+	}
+
+	if (URenderGridJob* DuplicateRenderGridJob = DuplicateObject(Job, this); IsValid(DuplicateRenderGridJob))
+	{
+		DuplicateRenderGridJob->GenerateNewId();
+		DuplicateRenderGridJob->SetJobId(GenerateNextJobId());
+		InsertRenderGridJobAfter(DuplicateRenderGridJob, Job);
+		return DuplicateRenderGridJob;
+	}
+	return nullptr;
+}
+
+bool URenderGrid::ReorderRenderGridJob(URenderGridJob* Job, URenderGridJob* DroppedOnJob, const bool bAfter)
+{
+	if (!IsValid(Job) || !IsValid(DroppedOnJob) || !HasRenderGridJob(Job) || !HasRenderGridJob(DroppedOnJob))
+	{
+		return false;
+	}
+
+	RemoveRenderGridJob(Job);
+	if (bAfter)
+	{
+		InsertRenderGridJobAfter(Job, DroppedOnJob);
+	}
+	else
+	{
+		InsertRenderGridJobBefore(Job, DroppedOnJob);
+	}
+	return true;
+}
diff --git a/Engine/Plugins/Experimental/RenderGrid/Source/RenderGrid/Private/RenderGrid/RenderGridManager.cpp b/Engine/Plugins/Experimental/RenderGrid/Source/RenderGrid/Private/RenderGrid/RenderGridManager.cpp
new file mode 100644
--- /dev/null
+++ b/Engine/Plugins/Experimental/RenderGrid/Source/RenderGrid/Private/RenderGrid/RenderGridManager.cpp
@@ -0,0 +1,252 @@
+// Copyright Epic Games, Inc. All Rights Reserved.
+
+#include "RenderGrid/RenderGridManager.h"
+#include "RenderGrid/RenderGridMoviePipelineJob.h"
+#include "RenderGridUtils.h"
+#include "MoviePipelineHighResSetting.h"
+#include "MoviePipelinePIEExecutor.h"
+
+
+URenderGridMoviePipelineRenderJob* UE::RenderGrid::FRenderGridManager::CreateBatchRenderJob(URenderGrid* Grid)
+{
+	FRenderGridMoviePipelineRenderJobCreateArgs JobArgs;
+	JobArgs.RenderGrid = Grid;
+	JobArgs.RenderGridJobs.Append(Grid->GetEnabledRenderGridJobs());
+	URenderGridMoviePipelineRenderJob* NewRenderJob = URenderGridMoviePipelineRenderJob::Create(JobArgs);
+	if (!IsValid(NewRenderJob))
+	{
+		return nullptr;
+	}
+	return NewRenderJob;
+}
+
+
+URenderGridMoviePipelineRenderJob* UE::RenderGrid::FRenderGridManager::RenderPreviewFrame(const FRenderGridManagerRenderPreviewFrameArgs& Args)
+{
+	const FRenderGridManagerRenderPreviewFrameArgsCallback Callback = Args.Callback;
+
+	if (!IsValid(Args.RenderGridJob))
+	{
+		Callback.ExecuteIfBound(false);
+		return nullptr;
+	}
+
+	URenderGridJob* JobCopy = DuplicateObject(Args.RenderGridJob.Get(), Args.RenderGridJob->GetOuter());
+	if (!IsValid(JobCopy))
+	{
+		Callback.ExecuteIfBound(false);
+		return nullptr;
+	}
+
+	JobCopy->SetJobId(JobCopy->GetId().ToString(EGuidFormats::Base36Encoded));
+
+	if (Args.Frame.IsSet())
+	{
+		constexpr int32 RenderFramesCount = 1;// can be more than 1 to prevent rendering issues, will always take the last frame that's rendered
+
+		JobCopy->SetIsUsingCustomStartFrame(true);
+		JobCopy->SetCustomStartFrame(Args.Frame.Get(0));
+
+		JobCopy->SetIsUsingCustomEndFrame(true);
+		JobCopy->SetCustomEndFrame(Args.Frame.Get(0));
+
+		if (!JobCopy->SetSequenceEndFrame(JobCopy->GetSequenceStartFrame().Get(0) + 1) ||
+			!JobCopy->SetSequenceStartFrame(JobCopy->GetSequenceEndFrame().Get(0) - RenderFramesCount))
+		{
+			Callback.ExecuteIfBound(false);
+			return nullptr;
+		}
+	}
+
+	JobCopy->SetIsUsingCustomResolution(true);
+	JobCopy->SetCustomResolution(Args.Resolution);
+
+	JobCopy->SetOutputDirectory(TmpRenderedFramesPath / (Args.Frame.IsSet() ? TEXT("PreviewFrame") : TEXT("PreviewFrames")));
+
+	FRenderGridMoviePipelineRenderJobCreateArgs JobArgs;
+	JobArgs.RenderGrid = Args.RenderGrid;
+	JobArgs.RenderGridJobs.Add(JobCopy);
+	JobArgs.bHeadless = Args.bHeadless;
+	JobArgs.bForceOutputImage = true;
+	JobArgs.bForceOnlySingleOutput = true;
+	JobArgs.bForceUseSequenceFrameRate = Args.Frame.IsSet();
+	JobArgs.bEnsureSequentialFilenames = true;
+	JobArgs.DisableSettingsClasses.Add(UMoviePipelineAntiAliasingSetting::StaticClass());
+	JobArgs.DisableSettingsClasses.Add(UMoviePipelineHighResSetting::StaticClass());
+
+	URenderGridMoviePipelineRenderJob* NewRenderJob = URenderGridMoviePipelineRenderJob::Create(JobArgs);
+	if (!NewRenderJob)
+	{
+		Callback.ExecuteIfBound(false);
+		return nullptr;
+	}
+
+	const FGuid JobId = JobCopy->GetId();
+	const TOptional<int32> StartFrameOfRender = (Args.Frame.IsSet() ? TOptional<int32>() : JobCopy->GetStartFrame());
+	NewRenderJob->OnExecuteFinished().AddLambda([this, Callback, JobId, StartFrameOfRender](URenderGridMoviePipelineRenderJob* RenderJob, const bool bSuccess)
+	{
+		if (StartFrameOfRender.IsSet())
+		{
+			if (bSuccess)
+			{
+				StartFrameOfRenders.Add(JobId, StartFrameOfRender.Get(0));
+			}
+			else
+			{
+				StartFrameOfRenders.Remove(JobId);
+			}
+		}
+		Callback.ExecuteIfBound(bSuccess);
+	});
+
+	NewRenderJob->Execute();
+	return NewRenderJob;
+}
+
+UTexture2D* UE::RenderGrid::FRenderGridManager::GetSingleRenderedPreviewFrame(URenderGridJob* Job, UTexture2D* ReusingTexture2D, bool& bOutReusedGivenTexture2D)
+{
+	static const FString PreviewFramesDir = TmpRenderedFramesPath / TEXT("PreviewFrame");
+	if (!IsValid(Job))
+	{
+		bOutReusedGivenTexture2D = (ReusingTexture2D == nullptr);
+		return nullptr;
+	}
+	const FString PreviewFramesSubDir = Job->GetId().ToString(EGuidFormats::Base36Encoded);
+
+	TArray<FString> ImagePaths = Private::FRenderGridUtils::GetFiles(PreviewFramesDir / PreviewFramesSubDir, true);
+	ImagePaths.Sort();
+	Algo::Reverse(ImagePaths);
+	for (const FString& ImagePath : ImagePaths)
+	{
+		if (UTexture2D* Texture = Private::FRenderGridUtils::GetImage(ImagePath, ReusingTexture2D, bOutReusedGivenTexture2D))
+		{
+			return Texture;
+		}
+	}
+	bOutReusedGivenTexture2D = (ReusingTexture2D == nullptr);
+	return nullptr;
+}
+
+UTexture2D* UE::RenderGrid::FRenderGridManager::GetSingleRenderedPreviewFrame(URenderGridJob* Job)
+{
+	bool bOutReusedGivenTexture2D;
+	return GetSingleRenderedPreviewFrame(Job, nullptr, bOutReusedGivenTexture2D);
+}
+
+UTexture2D* UE::RenderGrid::FRenderGridManager::GetRenderedPreviewFrame(URenderGridJob* Job, const int32 Frame, UTexture2D* ReusingTexture2D, bool& bOutReusedGivenTexture2D)
+{
+	static const FString PreviewFramesDir = TmpRenderedFramesPath / TEXT("PreviewFrames");
+	if (!IsValid(Job))
+	{
+		bOutReusedGivenTexture2D = (ReusingTexture2D == nullptr);
+		return nullptr;
+	}
+	const FString PreviewFramesSubDir = Job->GetId().ToString(EGuidFormats::Base36Encoded);
+
+	int32 CurrentFrame = 0;
+	if (int32* StartFrameOfRenderPtr = StartFrameOfRenders.Find(Job->GetId()))
+	{
+		CurrentFrame = *StartFrameOfRenderPtr;
+	}
+
+	TArray<FString> ImagePaths = Private::FRenderGridUtils::GetFiles(PreviewFramesDir / PreviewFramesSubDir, true);
+	for (const FString& ImagePath : ImagePaths)
+	{
+		if (!Private::FRenderGridUtils::IsImage(ImagePath))
+		{
+			continue;
+		}
+		if (CurrentFrame == Frame)
+		{
+			return Private::FRenderGridUtils::GetImage(ImagePath, ReusingTexture2D, bOutReusedGivenTexture2D);
+		}
+		CurrentFrame++;
+	}
+
+	bOutReusedGivenTexture2D = (ReusingTexture2D == nullptr);
+	return nullptr;
+}
+
+UTexture2D* UE::RenderGrid::FRenderGridManager::GetRenderedPreviewFrame(URenderGridJob* Job, const int32 Frame)
+{
+	bool bOutReusedGivenTexture2D;
+	return GetRenderedPreviewFrame(Job, Frame, nullptr, bOutReusedGivenTexture2D);
+}
+
+
+void UE::RenderGrid::FRenderGridManager::UpdateRenderGridJobsPropValues(URenderGrid* Grid)
+{
+	if (!IsValid(Grid))
+	{
+		return;
+	}
+
+	URenderGridPropsSourceRemoteControl* PropsSource = Grid->GetPropsSource<URenderGridPropsSourceRemoteControl>();
+	if (!IsValid(PropsSource))
+	{
+		return;
+	}
+
+	TArray<URenderGridJob*> Jobs = Grid->GetRenderGridJobs();
+	TArray<uint8> BinaryArray;
+	for (URenderGridPropRemoteControl* Field : PropsSource->GetProps()->GetAllCasted())
+	{
+		if (!Field->GetValue(BinaryArray))
+		{
+			continue;
+		}
+		const TSharedPtr<FRemoteControlEntity> Entity = Field->GetRemoteControlEntity();
+		for (URenderGridJob* Job : Jobs)
+		{
+			if (!Job->HasRemoteControlValue(Entity))
+			{
+				Job->SetRemoteControlValue(Entity, BinaryArray);
+			}
+		}
+	}
+}
+
+FRenderGridManagerPreviousPropValues UE::RenderGrid::FRenderGridManager::ApplyJobPropValues(const URenderGrid* Grid, const URenderGridJob* Job)
+{
+	FRenderGridManagerPreviousPropValues PreviousPropValues;
+
+	if (!IsValid(Grid) || !IsValid(Job))
+	{
+		return PreviousPropValues;
+	}
+
+	URenderGridPropsSourceBase* PropsSource = Grid->GetPropsSource();
+	if (URenderGridPropsSourceRemoteControl* PropsSourceRC = Cast<URenderGridPropsSourceRemoteControl>(PropsSource))
+	{
+		for (URenderGridPropRemoteControl* Prop : PropsSourceRC->GetProps()->GetAllCasted())
+		{
+			TArray<uint8> PreviousPropData;
+			if (!Prop->GetValue(PreviousPropData))
+			{
+				continue;
+			}
+
+			TArray<uint8> PropData;
+			if (!Job->ConstGetRemoteControlValue(Prop->GetRemoteControlEntity(), PropData))
+			{
+				continue;
+			}
+
+			PreviousPropValues.RemoteControlData.Add(Prop, PreviousPropData);
+			Prop->SetValue(PropData);
+		}
+		return PreviousPropValues;
+	}
+	return PreviousPropValues;
+}
+
+void UE::RenderGrid::FRenderGridManager::RestorePropValues(const FRenderGridManagerPreviousPropValues& PreviousPropValues)
+{
+	for (const auto& PropValuePair : PreviousPropValues.RemoteControlData)
+	{
+		if (IsValid(PropValuePair.Key))
+		{
+			PropValuePair.Key->SetValue(PropValuePair.Value.Bytes);
+		}
+	}
+}
diff --git a/Engine/Plugins/Experimental/RenderGrid/Source/RenderGrid/Private/RenderGrid/RenderGridMoviePipelineJob.cpp b/Engine/Plugins/Experimental/RenderGrid/Source/RenderGrid/Private/RenderGrid/RenderGridMoviePipelineJob.cpp
new file mode 100644
--- /dev/null
+++ b/Engine/Plugins/Experimental/RenderGrid/Source/RenderGrid/Private/RenderGrid/RenderGridMoviePipelineJob.cpp
@@ -0,0 +1,457 @@
+// Copyright Epic Games, Inc. All Rights Reserved.
+
+#include "RenderGrid/RenderGridMoviePipelineJob.h"
+
+#include "IRenderGridModule.h"
+#include "RenderGrid/RenderGrid.h"
+#include "RenderGrid/RenderGridManager.h"
+#include "RenderGridUtils.h"
+#include "Utils/RenderGridQueue.h"
+
+#include "LevelSequence.h"
+#include "LevelSequenceEditorModule.h"
+#include "LevelSequenceEditorModule.h"
+#include "Modules/ModuleManager.h"
+#include "MoviePipelineAntiAliasingSetting.h"
+#include "MoviePipelineDeferredPasses.h"
+#include "MoviePipelineEditorBlueprintLibrary.h"
+#include "MoviePipelineExecutor.h"
+#include "MoviePipelineImageSequenceOutput.h"
+#include "MoviePipelineOutputSetting.h"
+#include "MoviePipelinePIEExecutor.h"
+#include "MoviePipelineQueue.h"
+
+
+URenderGridMoviePipelineRenderJobEntry* URenderGridMoviePipelineRenderJobEntry::Create(URenderGridMoviePipelineRenderJob* RenderJob, URenderGridJob* Job, const UE::RenderGrid::FRenderGridMoviePipelineRenderJobCreateArgs& Args)
+{
+	if (!IsValid(RenderJob) || !IsValid(Job) || !IsValid(Args.RenderGrid) || (Args.RenderGridJobs.Num() <= 0))
+	{
+		return nullptr;
+	}
+
+	const UClass* PipelineExecutor = (IsValid(*Args.PipelineExecutorClass) ? *Args.PipelineExecutorClass : UMoviePipelinePIEExecutor::StaticClass());
+	if (!PipelineExecutor)
+	{
+		return nullptr;
+	}
+
+	URenderGridMoviePipelineRenderJobEntry* RenderJobEntry = NewObject<URenderGridMoviePipelineRenderJobEntry>(RenderJob);
+	RenderJobEntry->RenderGridJob = Job;
+	RenderJobEntry->RenderGrid = Args.RenderGrid;
+	RenderJobEntry->RenderQueue = NewObject<UMoviePipelineQueue>(RenderJobEntry);
+	RenderJobEntry->Executor = NewObject<UMoviePipelineExecutorBase>(RenderJobEntry, PipelineExecutor);
+	RenderJobEntry->ExecutorJob = nullptr;
+	RenderJobEntry->Status = TEXT("Skipped");
+	RenderJobEntry->bCanExecute = false;
+	RenderJobEntry->bCanceled = false;
+	RenderJobEntry->Promise = MakeShared<TPromise<void>>();
+	RenderJobEntry->Promise->SetValue();
+	RenderJobEntry->PromiseFuture = RenderJobEntry->Promise->GetFuture().Share();
+	RenderJobEntry->Promise.Reset();
+
+	if (Args.bHeadless)
+	{
+		if (UMoviePipelinePIEExecutor* ActiveExecutorPIE = Cast<UMoviePipelinePIEExecutor>(RenderJobEntry->Executor))
+		{
+			ActiveExecutorPIE->SetIsRenderingOffscreen(true);
+		}
+	}
+
+
+	ULevelSequence* JobSequence = Job->GetSequence();
+	if (!IsValid(JobSequence) || !Job->GetSequenceStartFrame().IsSet() || !Job->GetSequenceEndFrame().IsSet() || (Job->GetSequenceStartFrame().Get(0) >= Job->GetSequenceEndFrame().Get(0)))
+	{
+		return RenderJobEntry;
+	}
+
+	UMoviePipelineExecutorJob* NewJob = UMoviePipelineEditorBlueprintLibrary::CreateJobFromSequence(RenderJobEntry->RenderQueue, JobSequence);
+	RenderJobEntry->ExecutorJob = NewJob;
+
+	UMoviePipelineMasterConfig* JobRenderPreset = Job->GetRenderPreset();
+	if (IsValid(JobRenderPreset))
+	{
+		NewJob->SetConfiguration(JobRenderPreset);
+	}
+	else
+	{
+		UMoviePipelineEditorBlueprintLibrary::EnsureJobHasDefaultSettings(NewJob);
+	}
+
+	if (Args.DisableSettingsClasses.Num() > 0)
+	{
+		for (UMoviePipelineSetting* Setting : NewJob->GetConfiguration()->FindSettings<UMoviePipelineSetting>())
+		{
+			if (!IsValid(Setting))
+			{
+				continue;
+			}
+			for (TSubclassOf<UMoviePipelineSetting> DisableSettingsClass : Args.DisableSettingsClasses)
+			{
+				if (Setting->IsA(DisableSettingsClass))
+				{
+					Setting->SetIsEnabled(false);
+					break;
+				}
+			}
+		}
+	}
+
+	if (Args.bForceOutputImage || Args.bForceOnlySingleOutput)
+	{
+		bool bFound = false;
+		const bool bContainsPreferredType = IsValid(NewJob->GetConfiguration()->FindSetting<UMoviePipelineImageSequenceOutput_PNG>());
+		for (UMoviePipelineOutputBase* Setting : NewJob->GetConfiguration()->FindSettings<UMoviePipelineOutputBase>())
+		{
+			if (!IsValid(Setting))
+			{
+				continue;
+			}
+			if (Cast<UMoviePipelineImageSequenceOutput_PNG>(Setting) || Cast<UMoviePipelineImageSequenceOutput_JPG>(Setting) || Cast<UMoviePipelineImageSequenceOutput_BMP>(Setting))
+			{
+				if (!Args.bForceOnlySingleOutput || (!bFound && (!bContainsPreferredType || Cast<UMoviePipelineImageSequenceOutput_PNG>(Setting))))
+				{
+					bFound = true;
+					continue;
+				}
+			}
+			Setting->SetIsEnabled(false);
+		}
+		if (Args.bForceOutputImage && !bFound)
+		{
+			if (UMoviePipelineImageSequenceOutput_PNG* NewSetting = Cast<UMoviePipelineImageSequenceOutput_PNG>(NewJob->GetConfiguration()->FindOrAddSettingByClass(UMoviePipelineImageSequenceOutput_PNG::StaticClass())))
+			{
+				NewSetting->bWriteAlpha = false;
+			}
+		}
+	}
+
+	if (UMoviePipelineAntiAliasingSetting* ExistingAntiAliasingSettings = NewJob->GetConfiguration()->FindSetting<UMoviePipelineAntiAliasingSetting>())
+	{
+		// anti-aliasing settings already present (and enabled)
+		if ((UE::MovieRenderPipeline::GetEffectiveAntiAliasingMethod(ExistingAntiAliasingSettings) == AAM_FXAA) && (ExistingAntiAliasingSettings->SpatialSampleCount <= 1) && (ExistingAntiAliasingSettings->TemporalSampleCount <= 1))
+		{
+			ExistingAntiAliasingSettings->TemporalSampleCount = 2;// FXAA transparency fix
+		}
+	}
+	else if (UMoviePipelineAntiAliasingSetting* AntiAliasingSettings = Cast<UMoviePipelineAntiAliasingSetting>(NewJob->GetConfiguration()->FindOrAddSettingByClass(UMoviePipelineAntiAliasingSetting::StaticClass())))
+	{
+		// anti-aliasing settings not yet present (or enabled), created a new one
+		AntiAliasingSettings->EngineWarmUpCount = 0;
+		AntiAliasingSettings->RenderWarmUpCount = 0;
+		AntiAliasingSettings->SpatialSampleCount = 1;
+		AntiAliasingSettings->TemporalSampleCount = 2;
+		AntiAliasingSettings->bOverrideAntiAliasing = true;
+		AntiAliasingSettings->AntiAliasingMethod = AAM_FXAA;
+	}
+
+	bool bHasShot = false;
+	for (const UMoviePipelineExecutorShot* Shot : NewJob->ShotInfo)
+	{
+		if (!Shot)
+		{
+			continue;
+		}
+		bHasShot = true;
+
+		UMoviePipelineOutputSetting* Setting = Cast<UMoviePipelineOutputSetting>(UMoviePipelineBlueprintLibrary::FindOrGetDefaultSettingForShot(UMoviePipelineOutputSetting::StaticClass(), NewJob->GetConfiguration(), Shot));
+
+		Setting->bUseCustomPlaybackRange = true;
+		Setting->CustomStartFrame = Job->GetSequenceStartFrame().Get(0);
+		Setting->CustomEndFrame = Job->GetSequenceEndFrame().Get(0);
+
+		if (Args.bForceUseSequenceFrameRate)
+		{
+			Setting->bUseCustomFrameRate = false;
+		}
+
+		if (Job->GetIsUsingCustomResolution())
+		{
+			Setting->OutputResolution = Job->GetCustomResolution();
+		}
+
+		const FString JobOutputRootDirectory = Job->GetOutputDirectory();
+		const FString JobId = Job->GetJobId();
+		if (!JobOutputRootDirectory.IsEmpty() && !JobId.IsEmpty())
+		{
+			const FString JobOutputDirectory = JobOutputRootDirectory / JobId;
+			UE::RenderGrid::Private::FRenderGridUtils::DeleteDirectory(JobOutputDirectory);
+			Setting->OutputDirectory.Path = JobOutputDirectory;
+		}
+
+		if (Args.bEnsureSequentialFilenames || !IsValid(JobRenderPreset))
+		{
+			Setting->FileNameFormat = TEXT("{frame_number}");
+			Setting->ZeroPadFrameNumbers = 10;
+			Setting->FrameNumberOffset = 1000000000;
+		}
+	}
+	if (!bHasShot)
+	{
+		return RenderJobEntry;
+	}
+
+	RenderJobEntry->Status = TEXT("");
+	RenderJobEntry->bCanExecute = true;
+	return RenderJobEntry;
+}
+
+void URenderGridMoviePipelineRenderJobEntry::BeginDestroy()
+{
+	if (Promise.IsValid())
+	{
+		Promise->SetValue();
+		Promise.Reset();
+	}
+	Super::BeginDestroy();
+}
+
+TSharedFuture<void> URenderGridMoviePipelineRenderJobEntry::Execute()
+{
+	if (Executor->IsRendering())
+	{
+		return PromiseFuture;
+	}
+
+	Promise = MakeShared<TPromise<void>>();
+	PromiseFuture = Promise->GetFuture().Share();
+
+	if (!bCanExecute)
+	{
+		Status = TEXT("Skipped");
+		Promise->SetValue();
+		Promise.Reset();
+		return PromiseFuture;
+	}
+	if (bCanceled)
+	{
+		Status = TEXT("Canceled");
+		Promise->SetValue();
+		Promise.Reset();
+		return PromiseFuture;
+	}
+
+	AddToRoot();
+	if (ILevelSequenceEditorModule* LevelSequenceEditorModule = FModuleManager::GetModulePtr<ILevelSequenceEditorModule>("LevelSequenceEditor"); LevelSequenceEditorModule)
+	{
+		LevelSequenceEditorModule->OnComputePlaybackContext().AddUObject(this, &URenderGridMoviePipelineRenderJobEntry::ComputePlaybackContext);
+	}
+	if (UMoviePipelinePIEExecutor* ActiveExecutorPIE = Cast<UMoviePipelinePIEExecutor>(Executor))
+	{
+		ActiveExecutorPIE->OnIndividualJobStarted().AddUObject(this, &URenderGridMoviePipelineRenderJobEntry::ExecuteJobStarted);
+		ActiveExecutorPIE->OnIndividualJobWorkFinished().AddUObject(this, &URenderGridMoviePipelineRenderJobEntry::ExecuteJobFinished);
+	}
+	Status = TEXT("Rendering...");
+	Executor->OnExecutorFinished().AddUObject(this, &URenderGridMoviePipelineRenderJobEntry::ExecuteFinished);
+	Executor->Execute(RenderQueue);
+	return PromiseFuture;
+}
+
+void URenderGridMoviePipelineRenderJobEntry::Cancel()
+{
+	bCanceled = true;
+	if (Executor->IsRendering())
+	{
+		Executor->CancelAllJobs();
+	}
+}
+
+FString URenderGridMoviePipelineRenderJobEntry::GetStatus() const
+{
+	if (UMoviePipelineExecutorJob* RenderExecutorJob = ExecutorJob.Get(); IsValid(RenderExecutorJob))
+	{
+		if (FString JobStatus = RenderExecutorJob->GetStatusMessage().TrimStartAndEnd(); !JobStatus.IsEmpty())
+		{
+			return JobStatus;
+		}
+	}
+	return Status;
+}
+
+int32 URenderGridMoviePipelineRenderJobEntry::GetEngineWarmUpCount() const
+{
+	if (UMoviePipelineExecutorJob* RenderExecutorJob = ExecutorJob.Get(); IsValid(RenderExecutorJob))
+	{
+		if (UMoviePipelineAntiAliasingSetting* AntiAliasingSettings = Cast<UMoviePipelineAntiAliasingSetting>(RenderExecutorJob->GetConfiguration()->FindOrAddSettingByClass(UMoviePipelineAntiAliasingSetting::StaticClass())))
+		{
+			return FMath::Max<int32>(0, AntiAliasingSettings->EngineWarmUpCount);
+		}
+	}
+	return 0;
+}
+
+void URenderGridMoviePipelineRenderJobEntry::ComputePlaybackContext(bool& bOutAllowBinding)
+{
+	bOutAllowBinding = false;
+}
+
+void URenderGridMoviePipelineRenderJobEntry::ExecuteJobStarted(UMoviePipelineExecutorJob* StartingExecutorJob)
+{
+	RenderGrid->PreRender(RenderGridJob);
+}
+
+void URenderGridMoviePipelineRenderJobEntry::ExecuteJobFinished(FMoviePipelineOutputData PipelineOutputData)
+{
+	RenderGrid->PostRender(RenderGridJob);
+}
+
+void URenderGridMoviePipelineRenderJobEntry::ExecuteFinished(UMoviePipelineExecutorBase* PipelineExecutor, const bool bSuccess)
+{
+	if (ILevelSequenceEditorModule* LevelSequenceEditorModule = FModuleManager::GetModulePtr<ILevelSequenceEditorModule>("LevelSequenceEditorModule"))
+	{
+		LevelSequenceEditorModule->OnComputePlaybackContext().RemoveAll(this);
+	}
+	if (UMoviePipelinePIEExecutor* PipelineExecutorPIE = Cast<UMoviePipelinePIEExecutor>(PipelineExecutor))
+	{
+		PipelineExecutorPIE->OnIndividualJobStarted().RemoveAll(this);
+		PipelineExecutorPIE->OnIndividualJobWorkFinished().RemoveAll(this);
+	}
+	bCanceled = (bCanceled || !bSuccess);
+	Status = (bCanceled ? TEXT("Canceled") : TEXT("Done"));
+	if (Promise.IsValid())
+	{
+		Promise->SetValue();
+		Promise.Reset();
+	}
+	RemoveFromRoot();
+}
+
+
+URenderGridMoviePipelineRenderJob* URenderGridMoviePipelineRenderJob::Create(const UE::RenderGrid::FRenderGridMoviePipelineRenderJobCreateArgs& Args)
+{
+	if (!IsValid(Args.RenderGrid) || (Args.RenderGridJobs.Num() <= 0))
+	{
+		return nullptr;
+	}
+
+	const UClass* PipelineExecutor = (IsValid(*Args.PipelineExecutorClass) ? *Args.PipelineExecutorClass : UMoviePipelinePIEExecutor::StaticClass());
+	if (!PipelineExecutor)
+	{
+		return nullptr;
+	}
+
+	URenderGridMoviePipelineRenderJob* RenderJob = NewObject<URenderGridMoviePipelineRenderJob>(GetTransientPackage());
+	RenderJob->Queue = MakeShareable(new UE::RenderGrid::Private::FRenderGridQueue);
+	RenderJob->RenderGrid = Args.RenderGrid;
+	RenderJob->bCanceled = false;
+
+
+	RenderJob->Queue->Add(UE::RenderGrid::Private::FRenderGridQueueAction::CreateLambda([RenderJob]()
+	{
+		RenderJob->PreviousFrameLimitSettings = UE::RenderGrid::Private::FRenderGridUtils::DisableFpsLimit();
+	}));
+	RenderJob->Queue->DelayFrames(1);
+
+	for (const TObjectPtr<URenderGridJob> Job : Args.RenderGridJobs)
+	{
+		if (URenderGridMoviePipelineRenderJobEntry* Entry = URenderGridMoviePipelineRenderJobEntry::Create(RenderJob, Job, Args); IsValid(Entry))
+		{
+			RenderJob->Entries.Add(Job, Entry);
+
+			RenderJob->Queue->Add(UE::RenderGrid::Private::FRenderGridQueueActionReturningDelay::CreateLambda([RenderJob, Job]() -> UE::RenderGrid::Private::FRenderGridQueueDelay
+			{
+				if (!RenderJob->IsCanceled())
+				{
+					RenderJob->PreviousProps = UE::RenderGrid::IRenderGridModule::Get().GetManager().ApplyJobPropValues(RenderJob->RenderGrid, Job.Get());
+					return UE::RenderGrid::Private::FRenderGridQueueDelay::Frames(2 + Job->GetWaitFramesBeforeRendering());
+				}
+				return nullptr;
+			}));
+
+			RenderJob->Queue->Add(UE::RenderGrid::Private::FRenderGridQueueActionReturningDelayFuture::CreateLambda([Entry]()-> TSharedFuture<void>
+			{
+				return Entry->Execute();
+			}));
+
+			RenderJob->Queue->Add(UE::RenderGrid::Private::FRenderGridQueueActionReturningDelay::CreateLambda([RenderJob]() -> UE::RenderGrid::Private::FRenderGridQueueDelay
+			{
+				if (!RenderJob->PreviousProps.IsEmpty())
+				{
+					UE::RenderGrid::IRenderGridModule::Get().GetManager().RestorePropValues(RenderJob->PreviousProps);
+					RenderJob->PreviousProps = FRenderGridManagerPreviousPropValues();
+					return UE::RenderGrid::Private::FRenderGridQueueDelay::Frames(2);
+				}
+				return nullptr;
+			}));
+
+			RenderJob->Queue->Add(UE::RenderGrid::Private::FRenderGridQueueAction::CreateLambda([RenderJob, Entry]()
+			{
+				if (!RenderJob->IsCanceled() && Entry->IsCanceled())
+				{
+					RenderJob->Cancel();
+				}
+			}));
+		}
+	}
+
+	RenderJob->Queue->Add(UE::RenderGrid::Private::FRenderGridQueueAction::CreateLambda([RenderJob]()
+	{
+		UE::RenderGrid::Private::FRenderGridUtils::RestoreFpsLimit(RenderJob->PreviousFrameLimitSettings);
+		RenderJob->PreviousFrameLimitSettings = FRenderGridPreviousEngineFpsSettings();
+	}));
+
+
+	if (RenderJob->Entries.Num() <= 0)
+	{
+		return nullptr;
+	}
+	return RenderJob;
+}
+
+void URenderGridMoviePipelineRenderJob::Execute()
+{
+	if (Queue->IsRunning())
+	{
+		return;
+	}
+
+	{// start >>
+		OnExecuteStartedDelegate.Broadcast(this);
+		AddToRoot();
+	}// start <<
+
+	Queue->Add(UE::RenderGrid::Private::FRenderGridQueueAction::CreateLambda([this]()
+	{
+		{// end >>
+			RemoveFromRoot();
+			OnExecuteFinishedDelegate.Broadcast(this, !bCanceled);
+		}// end <<
+	}));
+
+	Queue->Start();
+}
+
+void URenderGridMoviePipelineRenderJob::Cancel()
+{
+	if (bCanceled)
+	{
+		return;
+	}
+	bCanceled = true;
+
+	TArray<TObjectPtr<URenderGridMoviePipelineRenderJobEntry>> EntryValues;
+	Entries.GenerateValueArray(EntryValues);
+	for (int64 i = EntryValues.Num() - 1; i >= 0; i--)
+	{
+		if (IsValid(EntryValues[i]))
+		{
+			EntryValues[i]->Cancel();
+		}
+	}
+}
+
+FString URenderGridMoviePipelineRenderJob::GetRenderGridJobStatus(URenderGridJob* Job) const
+{
+	if (!IsValid(Job))
+	{
+		return TEXT("");
+	}
+
+	if (const TObjectPtr<URenderGridMoviePipelineRenderJobEntry>* EntryPtr = Entries.Find(Job))
+	{
+		if (TObjectPtr<URenderGridMoviePipelineRenderJobEntry> Entry = *EntryPtr; IsValid(Entry))
+		{
+			return Entry->GetStatus();
+		}
+	}
+	return TEXT("");
+}
diff --git a/Engine/Plugins/Experimental/RenderGrid/Source/RenderGrid/Private/Utils/RenderGridQueue.cpp b/Engine/Plugins/Experimental/RenderGrid/Source/RenderGrid/Private/Utils/RenderGridQueue.cpp
new file mode 100644
--- /dev/null
+++ b/Engine/Plugins/Experimental/RenderGrid/Source/RenderGrid/Private/Utils/RenderGridQueue.cpp
@@ -0,0 +1,130 @@
+﻿// Copyright Epic Games, Inc. All Rights Reserved.
+
+#include "Utils/RenderGridQueue.h"
+
+
+void UE::RenderGrid::Private::FRenderGridQueue::Tick(float DeltaTime)
+{
+	if (DelayRemainingFrames > 0)
+	{
+		DelayRemainingFrames--;
+	}
+	if (DelayRemainingSeconds > 0)
+	{
+		DelayRemainingSeconds -= DeltaTime;
+	}
+
+	if (DelayRemainingFuture.IsValid())
+	{
+		if (!DelayRemainingFuture.IsReady())
+		{
+			return;
+		}
+		DelayRemainingFuture = TFuture<void>();
+	}
+	if (DelayRemainingFutureReturningDelay.IsValid())
+	{
+		if (!DelayRemainingFutureReturningDelay.IsReady())
+		{
+			return;
+		}
+		QueueDelay(DelayRemainingFutureReturningDelay.Get());
+		DelayRemainingFutureReturningDelay = TFuture<FRenderGridQueueDelay>();
+	}
+
+	if (!bExecuting || ((DelayRemainingFrames <= 0) && (DelayRemainingSeconds <= 0)))
+	{
+		ExecuteNext();
+	}
+}
+
+void UE::RenderGrid::Private::FRenderGridQueue::Start()
+{
+	bStarted = true;
+}
+
+void UE::RenderGrid::Private::FRenderGridQueue::Stop()
+{
+	bStarted = false;
+}
+
+void UE::RenderGrid::Private::FRenderGridQueue::ExecuteNext()
+{
+	bExecuting = true;
+	while (true)
+	{
+		if (DelayRemainingFuture.IsValid() || DelayRemainingFutureReturningDelay.IsValid())
+		{
+			// continue in Tick event
+			return;
+		}
+		if (ExecuteNextDelay())
+		{
+			// continue in Tick event
+			return;
+		}
+		if (bStarted && ExecuteNextEntry())
+		{
+			// execution completed, restart this loop  (to check for new delays, and to execute the next entry)
+			continue;
+		}
+		break;
+	}
+	// if it gets to this point, it means that there are no delays or executions anymore  (since it would have returned out of this function otherwise)
+	bExecuting = false;
+}
+
+bool UE::RenderGrid::Private::FRenderGridQueue::ExecuteNextDelay()
+{
+	while (true)
+	{
+		if (TDoubleLinkedList<FRenderGridQueueDelay>::TDoubleLinkedListNode* DelayNode = QueuedDelays.GetHead())
+		{
+			const FRenderGridQueueDelay& Delay = DelayNode->GetValue();
+			if (Delay.MinimumFrames > 0)
+			{
+				DelayRemainingFrames = Delay.MinimumFrames;
+			}
+			if (Delay.MinimumSeconds > 0)
+			{
+				DelayRemainingSeconds = Delay.MinimumSeconds;
+			}
+			QueuedDelays.RemoveNode(DelayNode);
+			if ((DelayRemainingFrames <= 0) && (DelayRemainingSeconds <= 0))
+			{
+				continue;
+			}
+			return true;
+		}
+		return false;
+	}
+}
+
+bool UE::RenderGrid::Private::FRenderGridQueue::ExecuteNextEntry()
+{
+	if (TDoubleLinkedList<FRenderGridQueueEntry>::TDoubleLinkedListNode* EntryNode = QueuedEntries.GetHead())
+	{
+		const FRenderGridQueueEntry& Entry = EntryNode->GetValue();
+
+		Entry.ActionRegular.ExecuteIfBound();
+
+		if (Entry.ActionReturningDelay.IsBound())
+		{
+			QueueDelay(Entry.ActionReturningDelay.Execute());
+		}
+
+		if (Entry.ActionReturningDelayFuture.IsBound())
+		{
+			DelayRemainingFuture = Entry.ActionReturningDelayFuture.Execute();
+		}
+
+		if (Entry.ActionReturningDelayFutureReturningDelay.IsBound())
+		{
+			DelayRemainingFutureReturningDelay = Entry.ActionReturningDelayFutureReturningDelay.Execute();
+		}
+
+		QueuedEntries.RemoveNode(EntryNode);
+		return true;
+	}
+	return false;
+}
diff --git a/Engine/Plugins/Experimental/RenderGrid/Source/RenderGrid/Private/Utils/RenderGridQueue.h b/Engine/Plugins/Experimental/RenderGrid/Source/RenderGrid/Private/Utils/RenderGridQueue.h
new file mode 100644
--- /dev/null
+++ b/Engine/Plugins/Experimental/RenderGrid/Source/RenderGrid/Private/Utils/RenderGridQueue.h
@@ -0,0 +1,175 @@
+﻿// Copyright Epic Games, Inc. All Rights Reserved.
+
+#pragma once
+
+#include "CoreMinimal.h"
+#include "Async/Future.h"
+#include "Containers/List.h"
+#include "Tickable.h"
+
+
+namespace UE::RenderGrid::Private
+{
+	/**
+	 * Struct containing the delay data (such as the frames or the seconds of delay this delay requires before it can move on to the next step).
+	 */
+	struct RENDERGRID_API FRenderGridQueueDelay
+	{
+	public:
+		static FRenderGridQueueDelay Frames(const int64 Frames) { return FRenderGridQueueDelay(Frames); }
+		static FRenderGridQueueDelay Seconds(const double Seconds) { return FRenderGridQueueDelay(Seconds); }
+		static FRenderGridQueueDelay FramesOrSeconds(const int64 Frames, const double Seconds) { return FRenderGridQueueDelay(Frames, Seconds); }
+
+	public:
+		FRenderGridQueueDelay(TYPE_OF_NULLPTR)
+		{}
+
+	private:
+		FRenderGridQueueDelay()
+		{}
+
+		explicit FRenderGridQueueDelay(const int64 Frames)
+			: MinimumFrames(Frames)
+		{}
+
+		explicit FRenderGridQueueDelay(const double Seconds)
+			: MinimumSeconds(Seconds)
+		{}
+
+		explicit FRenderGridQueueDelay(const int64 Frames, const double Seconds)
+			: MinimumFrames(Frames)
+			, MinimumSeconds(Seconds)
+		{}
+
+	public:
+		const int64 MinimumFrames = 0;
+		const double MinimumSeconds = 0.0f;
+	};
+
+
+	/** A delegate for a queued action. */
+	DECLARE_DELEGATE(FRenderGridQueueAction);
+
+	/** A delegate for a queued action, that optionally requires a delay after its execution. */
+	DECLARE_DELEGATE_RetVal(FRenderGridQueueDelay, FRenderGridQueueActionReturningDelay);
+
+	/** A delegate for a queued action, that will delay execution until the returned TFuture finishes. */
+	DECLARE_DELEGATE_RetVal(TSharedFuture<void>, FRenderGridQueueActionReturningDelayFuture);
+
+	/** A delegate for a queued action, that will delay execution until the returned TFuture finishes, which can optionally return yet another delay if required. */
+	DECLARE_DELEGATE_RetVal(TSharedFuture<FRenderGridQueueDelay>, FRenderGridQueueActionReturningDelayFutureReturningDelay);
+
+	/**
+	 * Struct containing the data of a queued action.
+	 */
+	struct RENDERGRID_API FRenderGridQueueEntry
+	{
+	public:
+		FRenderGridQueueEntry(const FRenderGridQueueAction& Action)
+			: ActionRegular(Action)
+		{}
+
+		FRenderGridQueueEntry(const FRenderGridQueueActionReturningDelay& Action)
+			: ActionReturningDelay(Action)
+		{}
+
+		FRenderGridQueueEntry(const FRenderGridQueueActionReturningDelayFuture& Action)
+			: ActionReturningDelayFuture(Action)
+		{}
+
+		FRenderGridQueueEntry(const FRenderGridQueueActionReturningDelayFutureReturningDelay& Action)
+			: ActionReturningDelayFutureReturningDelay(Action)
+		{}
+
+	public:
+		const FRenderGridQueueAction ActionRegular;
+		const FRenderGridQueueActionReturningDelay ActionReturningDelay;
+		const FRenderGridQueueActionReturningDelayFuture ActionReturningDelayFuture;
+		const FRenderGridQueueActionReturningDelayFutureReturningDelay ActionReturningDelayFutureReturningDelay;
+	};
+
+
+	/**
+	 * This class provides generic queue support, with built-in support for delays between actions.
+	 */
+	class RENDERGRID_API FRenderGridQueue : public TSharedFromThis<FRenderGridQueue>, public FTickableGameObject
+	{
+	public:
+		//~ Begin FTickableGameObject interface
+		virtual void Tick(float DeltaTime) override;
+		virtual ETickableTickType GetTickableTickType() const override { return ETickableTickType::Always; }
+		virtual bool IsTickableWhenPaused() const override { return true; }
+		virtual bool IsTickableInEditor() const override { return true; }
+		virtual bool IsTickable() const override { return true; }
+		virtual bool IsAllowedToTick() const override { return true; }
+		virtual TStatId GetStatId() const override
+		{
+			RETURN_QUICK_DECLARE_CYCLE_STAT(FRenderGridQueue, STATGROUP_Tickables);
+		}
+		//~ End FTickableGameObject interface
+
+	public:
+		/** Queues the given action. */
+		void Add(const FRenderGridQueueEntry& Entry) { QueuedEntries.AddTail(Entry); }
+
+		/** Queues the given delay. */
+		void Delay(const FRenderGridQueueDelay& Delay) { Add(FRenderGridQueueActionReturningDelay::CreateLambda([Delay]() -> FRenderGridQueueDelay { return Delay; })); }
+
+		/** Queues the given delay, which will wait for the given number of frames. */
+		void DelayFrames(const int64 Frames) { Delay(FRenderGridQueueDelay::Frames(Frames)); }
+
+		/** Queues the given delay, which will wait for the given number of seconds. */
+		void DelaySeconds(const double Seconds) { Delay(FRenderGridQueueDelay::Seconds(Seconds)); }
+
+		/** Queues the given delay, which will wait for the given number of frames or seconds, whatever takes the longest. */
+		void DelayFramesOrSeconds(const int64 Frames, const double Seconds) { Delay(FRenderGridQueueDelay::FramesOrSeconds(Frames, Seconds)); }
+
+		/** Starts the execution of this queue. */
+		void Start();
+
+		/** Stops (pauses) the execution of this queue, this can be resumed by calling the Start function again. Currently queued up delays will continue to run/expire. */
+		void Stop();
+
+		/** Returns true if Start() has been called and Stop() hasn't been called yet. */
+		bool IsRunning() { return bStarted; }
+
+	protected:
+		/** Executes the next delay (if there are any), otherwise it executes the next entry (action). */
+		void ExecuteNext();
+
+		/** Executes the next delay, returns true if it found any, returns false if there were no queued up delays. */
+		bool ExecuteNextDelay();
+
+		/** Executes the next entry (action), returns true if it found and executed an entry, returns false if there were no queued up entries. */
+		bool ExecuteNextEntry();
+
+		/** Adds the delay to the queued delays. */
+		void QueueDelay(const FRenderGridQueueDelay& Delay) { QueuedDelays.AddTail(Delay); }
+
+	protected:
+		/** The queued up entries (actions). */
+		TDoubleLinkedList<FRenderGridQueueEntry> QueuedEntries;
+
+		/** The queued up delays. */
+		TDoubleLinkedList<FRenderGridQueueDelay> QueuedDelays;
+
+	protected:
+		/** Whether it has started (and hasn't been stopped/paused yet). This means that if this is true, Start() has been called, and Stop() hasn't been called yet since then. */
+		bool bStarted = false;
+
+		/** Whether it's currently executing a delay or an entry (action). False means there were no delays and entries queued up anymore. */
+		bool bExecuting = false;
+
+		/** The number of frames the current delay has to wait for. */
+		int64 DelayRemainingFrames = 0;
+
+		/** The number of seconds the current delay has to wait for. */
+		double DelayRemainingSeconds = 0.0f;
+
+		/** The TFuture it's waiting for (if any). */
+		TSharedFuture<void> DelayRemainingFuture;
+
+		/** The TFuture it's waiting for (if any), that can return a delay. */
+		TSharedFuture<FRenderGridQueueDelay> DelayRemainingFutureReturningDelay;
+	};
+}
diff --git a/Engine/Plugins/Experimental/RenderGrid/Source/RenderGrid/Public/RenderGrid/RenderGrid.h b/Engine/Plugins/Experimental/RenderGrid/Source/RenderGrid/Public/RenderGrid/RenderGrid.h
new file mode 100644
--- /dev/null
+++ b/Engine/Plugins/Experimental/RenderGrid/Source/RenderGrid/Public/RenderGrid/RenderGrid.h
@@ -0,0 +1,470 @@
+// Copyright Epic Games, Inc. All Rights Reserved.
+
+#pragma once
+
+#include "CoreMinimal.h"
+#include "RenderGrid/RenderGridPropsSource.h"
+#include "RenderGrid.generated.h"
+
+
+class UMoviePipelineOutputSetting;
+class UMoviePipelineMasterConfig;
+class ULevelSequence;
+class URenderGrid;
+
+
+/**
+ * This struct contains the data for a certain remote control property.
+ * 
+ * It's currently simply a wrapper around a byte array.
+ * This struct is needed so that that byte array can be used in another UPROPERTY container (TMap, TArray, etc).
+ */
+USTRUCT(BlueprintType)
+struct RENDERGRID_API FRenderGridRemoteControlPropertyData
+{
+	GENERATED_BODY()
+
+public:
+	FRenderGridRemoteControlPropertyData() = default;
+	FRenderGridRemoteControlPropertyData(const TArray<uint8>& InBytes)
+		: Bytes(InBytes)
+	{}
+
+public:
+	/** The property data, as bytes. */
+	UPROPERTY()
+	TArray<uint8> Bytes;
+};
+
+
+/**
+ * This class represents a render grid job, in other words, an entry (a row) of a render grid.
+ * It contains a level sequence and custom properties that will be applied while rendering.
+ * 
+ * Each RenderGridJob must belong to a RenderGrid.
+ */
+UCLASS(BlueprintType)
+class RENDERGRID_API URenderGridJob : public UObject
+{
+	GENERATED_BODY()
+
+public:
+	URenderGridJob();
+
+	/** Gets the calculated start frame, not taking the framerate of the render preset into account. */
+	TOptional<int32> GetSequenceStartFrame() const;
+
+	/** Gets the calculated end frame, not taking the framerate of the render preset into account. */
+	TOptional<int32> GetSequenceEndFrame() const;
+
+	/** Sets the custom start frame to match the given sequence start frame. */
+	bool SetSequenceStartFrame(const int32 NewCustomStartFrame);
+
+	/** Sets the custom end frame to match the given sequence end frame. */
+	bool SetSequenceEndFrame(const int32 NewCustomEndFrame);
+
+	/** Gets the calculated start frame. */
+	TOptional<int32> GetStartFrame() const;
+
+	/** Gets the calculated end frame. */
+	TOptional<int32> GetEndFrame() const;
+
+	/** Gets the calculated start time. */
+	TOptional<double> GetStartTime() const;
+
+	/** Gets the calculated end time. */
+	TOptional<double> GetEndTime() const;
+
+	/** Gets the calculated duration in seconds. */
+	TOptional<double> GetDurationInSeconds() const;
+
+	/** Gets the aspect ratio that this job will be rendered in. */
+	UFUNCTION(BlueprintPure, Category="Render Grid|Job")
+	double GetOutputAspectRatio() const;
+
+	/** Checks whether the job contains data that matches the search terms. */
+	UFUNCTION(BlueprintPure, Category="Render Grid|Job")
+	bool MatchesSearchTerm(const FString& SearchTerm) const;
+
+public:
+	UFUNCTION(BlueprintPure, Category="Render Grid|Job")
+	FGuid GetId() const { return Id; }
+
+	UFUNCTION(BlueprintCallable, Category="Render Grid|Job")
+	void GenerateNewId() { Id = FGuid::NewGuid(); }
+
+
+	UFUNCTION(BlueprintPure, Category="Render Grid|Job")
+	int32 GetWaitFramesBeforeRendering() const { return WaitFramesBeforeRendering; }
+
+	UFUNCTION(BlueprintCallable, Category="Render Grid|Job")
+	void SetWaitFramesBeforeRendering(const int32 NewWaitFramesBeforeRendering) { WaitFramesBeforeRendering = FMath::Max<int32>(0, NewWaitFramesBeforeRendering); }
+
+
+	UFUNCTION(BlueprintPure, Category="Render Grid|Job")
+	ULevelSequence* GetSequence() const { return Sequence; }
+
+	UFUNCTION(BlueprintCallable, Category="Render Grid|Job")
+	void SetSequence(ULevelSequence* NewSequence) { Sequence = NewSequence; }
+
+
+	UFUNCTION(BlueprintPure, Category="Render Grid|Job")
+	bool GetIsUsingCustomStartFrame() const { return bOverrideStartFrame; }
+
+	UFUNCTION(BlueprintCallable, Category="Render Grid|Job")
+	void SetIsUsingCustomStartFrame(const bool bNewOverrideStartFrame) { bOverrideStartFrame = bNewOverrideStartFrame; }
+
+	UFUNCTION(BlueprintPure, Category="Render Grid|Job")
+	int32 GetCustomStartFrame() const { return CustomStartFrame; }
+
+	UFUNCTION(BlueprintCallable, Category="Render Grid|Job")
+	void SetCustomStartFrame(const int32 NewCustomStartFrame) { CustomStartFrame = NewCustomStartFrame; }
+
+
+	UFUNCTION(BlueprintPure, Category="Render Grid|Job")
+	bool GetIsUsingCustomEndFrame() const { return bOverrideEndFrame; }
+
+	UFUNCTION(BlueprintCallable, Category="Render Grid|Job")
+	void SetIsUsingCustomEndFrame(const bool bNewOverrideEndFrame) { bOverrideEndFrame = bNewOverrideEndFrame; }
+
+	UFUNCTION(BlueprintPure, Category="Render Grid|Job")
+	int32 GetCustomEndFrame() const { return CustomEndFrame; }
+
+	UFUNCTION(BlueprintCallable, Category="Render Grid|Job")
+	void SetCustomEndFrame(const int32 NewCustomEndFrame) { CustomEndFrame = NewCustomEndFrame; }
+
+
+	UFUNCTION(BlueprintPure, Category="Render Grid|Job")
+	bool GetIsUsingCustomResolution() const { return bOverrideResolution; }
+
+	UFUNCTION(BlueprintCallable, Category="Render Grid|Job")
+	void SetIsUsingCustomResolution(const bool bNewOverrideResolution) { bOverrideResolution = bNewOverrideResolution; }
+
+	UFUNCTION(BlueprintPure, Category="Render Grid|Job")
+	FIntPoint GetCustomResolution() const { return CustomResolution; }
+
+	UFUNCTION(BlueprintCallable, Category="Render Grid|Job")
+	void SetCustomResolution(const FIntPoint NewCustomResolution) { CustomResolution = NewCustomResolution; }
+
+
+	UFUNCTION(BlueprintPure, Category="Render Grid|Job")
+	FString GetJobId() const { return JobId; }
+
+	static FString PurgeJobIdOrReturnEmptyString(const FString& NewJobId);
+	static FString PurgeJobId(const FString& NewJobId);
+	static FString PurgeJobIdOrGenerateUniqueId(URenderGrid* Grid, const FString& NewJobId);
+
+	UFUNCTION(BlueprintCallable, Category="Render Grid|Job")
+	void SetJobId(const FString& NewJobId) { JobId = PurgeJobId(NewJobId); }
+	void SetJobIdRaw(const FString& NewJobId) { JobId = NewJobId; }
+
+
+	UFUNCTION(BlueprintPure, Category="Render Grid|Job")
+	FString GetJobName() const { return JobName; }
+
+	static FString PurgeJobName(const FString& NewJobName);
+
+	UFUNCTION(BlueprintCallable, Category="Render Grid|Job")
+	void SetJobName(const FString& NewJobName) { JobName = NewJobName; }
+	void SetJobNameRaw(const FString& NewJobName) { JobName = NewJobName; }
+
+
+	UFUNCTION(BlueprintPure, Category="Render Grid|Job")
+	bool GetIsEnabled() const { return bIsEnabled; }
+
+	UFUNCTION(BlueprintCallable, Category="Render Grid|Job")
+	void SetIsEnabled(const bool bEnabled) { bIsEnabled = bEnabled; }
+
+
+	UFUNCTION(BlueprintPure, Category="Render Grid|Job")
+	FString GetOutputDirectory() const;
+	FString GetOutputDirectoryRaw() const { return OutputDirectory; }
+
+	UFUNCTION(BlueprintPure, Category="Render Grid|Job")
+	FString GetOutputDirectoryForDisplay() const { return OutputDirectory; }
+
+	static FString PurgeOutputDirectory(const FString& NewOutputDirectory);
+
+	UFUNCTION(BlueprintCallable, Category="Render Grid|Job")
+	void SetOutputDirectory(const FString& NewOutputDirectory) { OutputDirectory = PurgeOutputDirectory(NewOutputDirectory); }
+	void SetOutputDirectoryRaw(const FString& NewOutputDirectory) { OutputDirectory = NewOutputDirectory; }
+
+
+	UFUNCTION(BlueprintPure, Category="Render Grid|Job")
+	UMoviePipelineMasterConfig* GetRenderPreset() const { return RenderPreset; }
+
+	UFUNCTION(BlueprintPure, Category="Render Grid|Job")
+	UMoviePipelineOutputSetting* GetRenderPresetOutputSettings() const;
+
+	UFUNCTION(BlueprintCallable, Category="Render Grid|Job")
+	void SetRenderPreset(UMoviePipelineMasterConfig* NewRenderPreset) { RenderPreset = NewRenderPreset; }
+
+
+	bool HasRemoteControlValue(const TSharedPtr<FRemoteControlEntity>& RemoteControlEntity) const;
+	bool ConstGetRemoteControlValue(const TSharedPtr<FRemoteControlEntity>& RemoteControlEntity, TArray<uint8>& OutBinaryArray) const;
+	bool GetRemoteControlValue(const TSharedPtr<FRemoteControlEntity>& RemoteControlEntity, TArray<uint8>& OutBinaryArray);
+	bool SetRemoteControlValue(const TSharedPtr<FRemoteControlEntity>& RemoteControlEntity, const TArray<uint8>& BinaryArray);
+	TMap<FString, FRenderGridRemoteControlPropertyData>& GetRemoteControlValuesRef() { return RemoteControlValues; }
+
+private:
+	/** The unique ID of this job. */
+	UPROPERTY()
+	FGuid Id;
+
+	/** Waits the given number of frames before it will render this job. This can be set to a higher amount when the renderer has to wait for your code to complete (such as construction scripts etc). Try increasing this value when rendering doesn't produce the output you expect it to. */
+	UPROPERTY(EditInstanceOnly, Category="Render Grid|Job", Meta=(AllowPrivateAccess="true", ClampMin="0"))
+	int32 WaitFramesBeforeRendering;
+
+	/** The level sequence, this is what will be rendered during rendering. A job without a level sequence can't be rendered. */
+	UPROPERTY(EditInstanceOnly, Category="Render Grid|Job", Meta=(AllowPrivateAccess="true"))
+	TObjectPtr<ULevelSequence> Sequence;
+
+	/** If this is true, the CustomStartFrame property will override the start frame of the level sequence. */
+	UPROPERTY(EditInstanceOnly, Category="Render Grid|Job", Meta=(AllowPrivateAccess="true", InlineEditConditionToggle))
+	bool bOverrideStartFrame;
+
+	/** If bOverrideStartFrame is true, this property will override the start frame of the level sequence. */
+	UPROPERTY(EditInstanceOnly, Category="Render Grid|Job", Meta=(AllowPrivateAccess="true", EditCondition="bOverrideStartFrame"))
+	int32 CustomStartFrame;
+
+	/** If this is true, the CustomEndFrame property will override the end frame of the level sequence. */
+	UPROPERTY(EditInstanceOnly, Category="Render Grid|Job", Meta=(AllowPrivateAccess="true", InlineEditConditionToggle))
+	bool bOverrideEndFrame;
+
+	/** If bOverrideEndFrame is true, this property will override the end frame of the level sequence. */
+	UPROPERTY(EditInstanceOnly, Category="Render Grid|Job", Meta=(AllowPrivateAccess="true", EditCondition="bOverrideEndFrame"))
+	int32 CustomEndFrame;
+
+	/** If this is true, the CustomResolution property will override the resolution of the render. */
+	UPROPERTY()
+	bool bOverrideResolution;
+
+	/** If bOverrideResolution is true, this property will override the resolution of the render. */
+	UPROPERTY()
+	FIntPoint CustomResolution;
+
+	/** If this is true, this job will be rendered during a batch rendering, otherwise it will be skipped. */
+	UPROPERTY()
+	bool bIsEnabled;
+
+	/** The 'ID' of this job, this 'ID' is set by users. During rendering it will place all the output files of this job into a folder called after this 'ID', this means that this string can only contain file-safe characters. */
+	UPROPERTY()
+	FString JobId;
+
+	/** The name of this job, this can be anything, it's set by the user, it serves as a way for the user to understand what job this is. */
+	UPROPERTY()
+	FString JobName;
+
+	/** This is the folder in which the output files (of rendering) are placed into. To be more specific, the output files are placed in: {OutputDirectory}/{JobId}/, this folder will be created if it doesn't exist at the time of rendering. */
+	UPROPERTY()
+	FString OutputDirectory;
+
+	/** The MRQ render preset. Render grid jobs are rendered using the MRQ (Movie Render Queue) plugin. This 'preset' contains the configuration of that plugin. */
+	UPROPERTY()
+	TObjectPtr<UMoviePipelineMasterConfig> RenderPreset;
+
+	/** The Remote Control plugin can be used to customize and modify the way a job is rendered. If Remote Control is being used, the property values of this job will be stored in this map (remote control entity id -> value as bytes). */
+	UPROPERTY()
+	TMap<FString, FRenderGridRemoteControlPropertyData> RemoteControlValues;
+};
+
+
+/**
+ * This class represents a render grid, which is basically just a collection of render grid jobs.
+ * A render grid is the asset that is shown in the content browser, it's the asset that can be opened and edited using the editor.
+ */
+UCLASS(Blueprintable, BlueprintType, EditInlineNew, Meta=(DontUseGenericSpawnObject="true"))
+class RENDERGRID_API URenderGrid : public UObject
+{
+	GENERATED_BODY()
+
+public:
+	URenderGrid();
+
+	//UObject interface
+	virtual UWorld* GetWorld() const override;
+	virtual void PreSave(FObjectPreSaveContext SaveContext) override;
+	virtual void PostLoad() override;
+	//~ End UObject Interface
+
+	/** Should be called when the editor closes this asset. */
+	void OnClose() { SaveValuesToCDO(); }
+
+public:
+	static TArray<FString> GetBlueprintImplementableEvents() { return {TEXT("ReceivePreRender"), TEXT("ReceivePostRender")}; }
+
+protected:
+	/** Event for when rendering begins for a job. */
+	UFUNCTION(BlueprintImplementableEvent, Meta=(DisplayName="PreRender"))
+	void ReceivePreRender(URenderGridJob* Job);
+
+	/** Event for when rendering ends for a job. */
+	UFUNCTION(BlueprintImplementableEvent, Meta=(DisplayName="PostRender"))
+	void ReceivePostRender(URenderGridJob* Job);
+
+public:
+	/** Overridable native event for when rendering begins for a job. */
+	virtual void PreRender(URenderGridJob* Job);
+
+	/** Overridable native event for when rendering ends for a job. */
+	virtual void PostRender(URenderGridJob* Job);
+
+private:
+	/**
+	 * Because render grid assets are blueprints (assets that also have a blueprint graph), the render grid data is not stored directly in the asset data that you see in the content browser.
+	 * Instead, the data that is stored (and load) is the CDO (class default object).
+	 * Because of that, any data that needs to persist needs to be copied over to the CDO during a save, and data you'd like to load from it needs to be copied from the CDO during a load.
+	 */
+
+	/** Obtains the CDO, could return itself if this is called on the CDO instance. */
+	URenderGrid* GetCDO() { return (HasAnyFlags(RF_ClassDefaultObject) ? this : GetClass()->GetDefaultObject<URenderGrid>()); }
+
+	/** Copied values over into the CDO. */
+	void SaveValuesToCDO() { CopyValuesToOrFromCDO(true); }
+
+	/** Copied values over from the CDO. */
+	void LoadValuesFromCDO() { CopyValuesToOrFromCDO(false); }
+
+	/** Copied values to or from the CDO, based on whether bToCDO is true or false. */
+	void CopyValuesToOrFromCDO(const bool bToCDO);
+
+public:
+	UFUNCTION(BlueprintPure, Category="Render Grid|Grid")
+	FGuid GetId() const { return Id; }
+
+	UFUNCTION(BlueprintCallable, Category="Render Grid|Grid")
+	void GenerateNewId() { Id = FGuid::NewGuid(); }
+
+	UFUNCTION(BlueprintCallable, Category="Render Grid|Grid")
+	void SetPropsSource(ERenderGridPropsSourceType InPropsSourceType, UObject* InPropsSourceOrigin = nullptr);
+
+	UFUNCTION(BlueprintPure, Category="Render Grid|Grid")
+	URenderGridPropsSourceBase* GetPropsSource() const;
+
+	template<typename Type>
+	Type* GetPropsSource() const
+	{
+		static_assert(TIsDerivedFrom<Type, URenderGridPropsSourceBase>::IsDerived, "Type needs to derive from URenderGridPropsSourceBase.");
+		return Cast<Type>(GetPropsSource());
+	}
+
+	UFUNCTION(BlueprintPure, Category="Render Grid|Grid")
+	ERenderGridPropsSourceType GetPropsSourceType() const { return PropsSourceType; }
+
+	UFUNCTION(BlueprintPure, Category="Render Grid|Grid")
+	UObject* GetPropsSourceOrigin() const;
+
+public:
+	UFUNCTION(BlueprintCallable, Category="Render Grid|Grid")
+	void AddRenderGridJob(URenderGridJob* Job);
+
+	UFUNCTION(BlueprintCallable, Category="Render Grid|Grid")
+	void RemoveRenderGridJob(URenderGridJob* Job);
+
+	UFUNCTION(BlueprintCallable, Category="Render Grid|Grid")
+	void InsertRenderGridJob(URenderGridJob* Job, int32 Index);
+
+	UFUNCTION(BlueprintCallable, Category="Render Grid|Grid")
+	bool HasAnyRenderGridJobs() const { return !RenderGridJobs.IsEmpty(); }
+
+	UFUNCTION(BlueprintPure, Category="Render Grid|Grid")
+	bool HasRenderGridJob(URenderGridJob* Job) const;
+
+	UFUNCTION(BlueprintPure, Category="Render Grid|Grid")
+	int32 GetIndexOfRenderGridJob(URenderGridJob* Job) const;
+
+	TArray<TObjectPtr<URenderGridJob>>& GetRenderGridJobsRef() { return RenderGridJobs; }
+
+	UFUNCTION(BlueprintPure, Category="Render Grid|Grid")
+	TArray<URenderGridJob*> GetRenderGridJobs() const;
+
+	UFUNCTION(BlueprintPure, Category="Render Grid|Grid")
+	TArray<URenderGridJob*> GetEnabledRenderGridJobs() const;
+
+	UFUNCTION(BlueprintPure, Category="Render Grid|Grid")
+	TArray<URenderGridJob*> GetDisabledRenderGridJobs() const;
+
+	UFUNCTION(BlueprintCallable, Category="Render Grid|Grid")
+	void InsertRenderGridJobBefore(URenderGridJob* Job, URenderGridJob* BeforeJob);
+
+	UFUNCTION(BlueprintCallable, Category="Render Grid|Grid")
+	void InsertRenderGridJobAfter(URenderGridJob* Job, URenderGridJob* AfterJob);
+
+public:
+	/** Generates a unique job ID by finding the currently highest job ID and increasing it by one. **/
+	UFUNCTION(BlueprintCallable, Category="Render Grid|Grid")
+	FString GenerateNextJobId();
+
+	/** Finds whether given job ID already exists in this grid. **/
+	UFUNCTION(BlueprintPure, Category="Render Grid|Grid")
+	bool DoesJobIdExist(const FString& JobId);
+
+	/** Creates a new job and adds it to this grid. **/
+	UFUNCTION(BlueprintCallable, Category="Render Grid|Grid")
+	URenderGridJob* CreateAndAddNewRenderGridJob();
+
+	/** Copy the given job in this grid. **/
+	UFUNCTION(BlueprintCallable, Category="Render Grid|Grid")
+	URenderGridJob* DuplicateAndAddRenderGridJob(URenderGridJob* Job);
+
+	/** Relocates the given job in this grid to the position of the given dropped-on job. **/
+	UFUNCTION(BlueprintCallable, Category="Render Grid|Grid")
+	bool ReorderRenderGridJob(URenderGridJob* Job, URenderGridJob* DroppedOnJob, const bool bAfter = true);
+
+public:
+	/** Returns true when it's currently executing a blueprint implementable event, returns false otherwise. */
+	bool IsCurrentlyExecutingUserCode() const { return bExecutingPreRender || bExecutingPostRender; }
+
+private:
+	DECLARE_MULTICAST_DELEGATE(FOnRenderGridPreSave);
+	FOnRenderGridPreSave& OnPreSaveCDO() { return GetCDO()->OnRenderGridPreSaveDelegate; }
+
+private:
+	/** The delegate for when this grid is about to save. */
+	FOnRenderGridPreSave OnRenderGridPreSaveDelegate;
+
+private:
+	/** The unique ID of this render grid. */
+	UPROPERTY()
+	FGuid Id;
+
+
+	/** The type of the properties that a job in this grid can have. */
+	UPROPERTY(EditInstanceOnly, Category="Render Grid|Grid", Meta=(DisplayName="Properties Type", AllowPrivateAccess="true"))
+	ERenderGridPropsSourceType PropsSourceType;
+
+	/** The remote control properties that a job in this grid can have, only use this if PropsSourceType is ERenderGridPropsSourceType::RemoteControl. */
+	UPROPERTY(EditInstanceOnly, Category="Render Grid|Grid", Meta=(DisplayName="Remote Control Preset", AllowPrivateAccess="true", EditCondition="PropsSourceType == ERenderGridPropsSourceType::RemoteControl", EditConditionHides))
+	TObjectPtr<URemoteControlPreset> PropsSourceOrigin_RemoteControl;
+
+
+	/** The jobs of this grid. */
+	UPROPERTY(Instanced)
+	TArray<TObjectPtr<URenderGridJob>> RenderGridJobs;
+
+
+	/** True when it's currently executing the PreRender event, false otherwise. */
+	UPROPERTY(Transient)
+	bool bExecutingPreRender;
+
+	/** True when it's currently executing the PostRender event, false otherwise. */
+	UPROPERTY(Transient)
+	bool bExecutingPostRender;
+
+
+	/** GetPropsSource calls are somewhat expensive, we speed that up by caching the result (the PropsSource) that has been last outputted by that function. */
+	UPROPERTY(Transient)
+	mutable TObjectPtr<URenderGridPropsSourceBase> CachedPropsSource;
+
+	/** GetPropsSource calls are somewhat expensive, we speed that up by caching the PropsSourceType last used in that function. */
+	UPROPERTY(Transient)
+	mutable ERenderGridPropsSourceType CachedPropsSourceType;
+
+	/** GetPropsSource calls are somewhat expensive, we speed that up by caching the PropsSourceOrigin last used in that function. */
+	UPROPERTY(Transient)
+	mutable TWeakObjectPtr<UObject> CachedPropsSourceOriginWeakPtr;
+
+
+	/** GetWorld calls can be expensive, we speed them up by caching the last found world until it goes away. */
+	mutable TWeakObjectPtr<UWorld> CachedWorldWeakPtr;
+};
diff --git a/Engine/Plugins/Experimental/RenderGrid/Source/RenderGrid/Public/RenderGrid/RenderGridMoviePipelineJob.h b/Engine/Plugins/Experimental/RenderGrid/Source/RenderGrid/Public/RenderGrid/RenderGridMoviePipelineJob.h
new file mode 100644
--- /dev/null
+++ b/Engine/Plugins/Experimental/RenderGrid/Source/RenderGrid/Public/RenderGrid/RenderGridMoviePipelineJob.h
@@ -0,0 +1,207 @@
+// Copyright Epic Games, Inc. All Rights Reserved.
+
+#pragma once
+
+#include "RenderGridManager.h"
+#include "RenderGridUtils.h"
+#include "RenderGridMoviePipelineJob.generated.h"
+
+
+class UMoviePipelineSetting;
+class UMoviePipelineOutputBase;
+struct FMoviePipelineOutputData;
+class URenderGrid;
+class URenderGridJob;
+class UMoviePipelineExecutorBase;
+class UMoviePipelinePIEExecutor;
+class UMoviePipelineQueue;
+class UMoviePipelineExecutorJob;
+class URenderGridMoviePipelineRenderJob;
+
+namespace UE::RenderGrid::Private
+{
+	class FRenderGridQueue;
+}
+
+
+namespace UE::RenderGrid
+{
+	/** A delegate for when a render job is about to start. */
+	DECLARE_MULTICAST_DELEGATE_OneParam(FOnRenderGridMoviePipelineRenderJobStarted, URenderGridMoviePipelineRenderJob* /*RenderJob*/);
+
+	/** A delegate for when a render job has finished. */
+	DECLARE_MULTICAST_DELEGATE_TwoParams(FOnRenderGridMoviePipelineRenderJobFinished, URenderGridMoviePipelineRenderJob* /*RenderJob*/, bool /*bSuccess*/);
+
+	/**
+	 * The arguments for the URenderGridMoviePipelineRenderJob::Create function.
+	 */
+	struct RENDERGRID_API FRenderGridMoviePipelineRenderJobCreateArgs
+	{
+	public:
+		/** The render grid of the given render grid jobs that will be rendered. */
+		TObjectPtr<URenderGrid> RenderGrid = nullptr;
+
+		/** The specific render grid jobs that will be rendered. */
+		TArray<TObjectPtr<URenderGridJob>> RenderGridJobs;
+
+		/** If not null, it will override the MRQ pipeline executor class with this class. */
+		TSubclassOf<UMoviePipelinePIEExecutor> PipelineExecutorClass = nullptr;
+
+		/** The MRQ settings classes to disable (things like Anti-Aliasing, High-Res, etc). */
+		TArray<TSubclassOf<UMoviePipelineSetting>> DisableSettingsClasses;
+
+		/** Whether it should run invisibly (so without any UI elements popping up during rendering) or not. */
+		bool bHeadless = false;
+
+		/** Whether it should make sure it will output an image or not (if this bool is true, it will test if JPG/PNG/etc output is enabled, if none are, it will enable PNG output). */
+		bool bForceOutputImage = false;
+
+		/** Whether it should make sure it will only output in a single format (if this bool is true, if for example JPG and PNG output are enabled, one will be disabled, so that there will only be 1 output that's enabled). */
+		bool bForceOnlySingleOutput = false;
+
+		/** Whether it should use the sequence's framerate rather than any manually set framerate (if this bool is true, it will make sure bUseCustomFrameRate is set to false). */
+		bool bForceUseSequenceFrameRate = false;
+
+		/** Whether it should make sure it will output files named 0000000001, 0000000002, etc (if this bool is true, it will override the FileNameFormat to simply output the frame number, and it will add 1000000000 to that frame number to hopefully ensure that any negative frame numbers will not result in filenames starting with a minus character). */
+		bool bEnsureSequentialFilenames = false;
+	};
+}
+
+
+/**
+ * This class is responsible for the MRQ part of the rendering of the given render grid job.
+ */
+UCLASS()
+class RENDERGRID_API URenderGridMoviePipelineRenderJobEntry : public UObject
+{
+	GENERATED_BODY()
+
+public:
+	/** Creates a new render job instance, it won't be started right away. */
+	static URenderGridMoviePipelineRenderJobEntry* Create(URenderGridMoviePipelineRenderJob* RenderJob, URenderGridJob* Job, const UE::RenderGrid::FRenderGridMoviePipelineRenderJobCreateArgs& Args);
+
+	/** The destructor, cleans up the TPromise (if it's set). */
+	virtual void BeginDestroy() override;
+
+	/** Starts this render job. */
+	TSharedFuture<void> Execute();
+
+	/** Cancels this render job. Relies on the internal MRQ implementation of job canceling on whether this will do anything or not. */
+	void Cancel();
+
+	/** Retrieves the rendering status of the given render grid job. */
+	FString GetStatus() const;
+
+	/** Retrieves the "Engine Warm Up Count" value from the AntiAliasingSettings from the render preset that this render grid job uses. */
+	int32 GetEngineWarmUpCount() const;
+
+public:
+	/** Returns true if this render job was canceled (which for example can be caused by calling Cancel(), or by closing the render popup). */
+	bool IsCanceled() const { return bCanceled; }
+
+private:
+	void ComputePlaybackContext(bool& bOutAllowBinding);
+	void ExecuteJobStarted(UMoviePipelineExecutorJob* StartingExecutorJob);
+	void ExecuteJobFinished(FMoviePipelineOutputData PipelineOutputData);
+	void ExecuteFinished(UMoviePipelineExecutorBase* PipelineExecutor, const bool bSuccess);
+
+protected:
+	/** The render grid job that will be rendered. */
+	UPROPERTY(Transient)
+	TObjectPtr<URenderGridJob> RenderGridJob;
+
+	/** The render grid that the render grid job (that will be rendered) belongs to. */
+	UPROPERTY(Transient)
+	TObjectPtr<URenderGrid> RenderGrid;
+
+	/** The MRQ queue. */
+	UPROPERTY(Transient)
+	TObjectPtr<UMoviePipelineQueue> RenderQueue;
+
+	/** The MRQ pipeline executor. */
+	UPROPERTY(Transient)
+	TObjectPtr<UMoviePipelineExecutorBase> Executor;
+
+	/** The MRQ job of the given render grid job. */
+	UPROPERTY(Transient)
+	TObjectPtr<UMoviePipelineExecutorJob> ExecutorJob;
+
+	/** The TPromise of the rendering process. */
+	TSharedPtr<TPromise<void>> Promise;
+
+	/** The TFuture of the rendering process. */
+	TSharedFuture<void> PromiseFuture;
+
+	/** The rendering status of the given render grid job. */
+	UPROPERTY(Transient)
+	FString Status;
+
+	/** Whether the entry can execute, or whether it should just skip execution. */
+	UPROPERTY(Transient)
+	bool bCanExecute;
+
+	/** Whether the entry was canceled (like by calling Cancel(), or by closing the render popup). */
+	UPROPERTY(Transient)
+	bool bCanceled;
+};
+
+
+/**
+ * This class is responsible for rendering the given render grid jobs.
+ */
+UCLASS()
+class RENDERGRID_API URenderGridMoviePipelineRenderJob : public UObject
+{
+	GENERATED_BODY()
+
+public:
+	/** Creates a new render job instance, it won't be started right away. */
+	static URenderGridMoviePipelineRenderJob* Create(const UE::RenderGrid::FRenderGridMoviePipelineRenderJobCreateArgs& Args);
+
+	/** Starts this render job. */
+	void Execute();
+
+	/** Cancels this render job. Relies on the internal MRQ implementation of job canceling on whether this will stop the current render grid job from rendering or not. Will always prevent new render grid jobs from rendering. */
+	void Cancel();
+
+	/** Returns true if this render job has been canceled. */
+	bool IsCanceled() const { return bCanceled; }
+
+	/** Retrieves the rendering status of the given render grid job. */
+	FString GetRenderGridJobStatus(URenderGridJob* Job) const;
+
+protected:
+	/** The queue containing the render actions. */
+	TSharedPtr<UE::RenderGrid::Private::FRenderGridQueue> Queue;
+
+	/** The render grid jobs that are to be rendered, mapped to the rendering job of each specific render grid job. */
+	UPROPERTY(Transient)
+	TMap<TObjectPtr<const URenderGridJob>, TObjectPtr<URenderGridMoviePipelineRenderJobEntry>> Entries;
+
+	/** The render grid of the given render grid job that will be rendered. */
+	UPROPERTY(Transient)
+	TObjectPtr<URenderGrid> RenderGrid;
+
+	/** Whether the remaining render grid jobs should be prevented from rendering. */
+	UPROPERTY(Transient)
+	bool bCanceled;
+
+	/** The property values that have been overwritten by the currently applied render grid job property values. */
+	UPROPERTY(Transient)
+	FRenderGridManagerPreviousPropValues PreviousProps;
+
+	/** The engine framerate settings values that have been overwritten by the currently applied engine framerate settings values. */
+	UPROPERTY(Transient)
+	FRenderGridPreviousEngineFpsSettings PreviousFrameLimitSettings;
+
+public:
+	/** A delegate for when the render job is about to start. */
+	UE::RenderGrid::FOnRenderGridMoviePipelineRenderJobStarted& OnExecuteStarted() { return OnExecuteStartedDelegate; }
+
+	/** A delegate for when the render job has finished. */
+	UE::RenderGrid::FOnRenderGridMoviePipelineRenderJobFinished& OnExecuteFinished() { return OnExecuteFinishedDelegate; }
+
+private:
+	UE::RenderGrid::FOnRenderGridMoviePipelineRenderJobStarted OnExecuteStartedDelegate;
+	UE::RenderGrid::FOnRenderGridMoviePipelineRenderJobFinished OnExecuteFinishedDelegate;
+};
diff --git a/Engine/Plugins/Experimental/RenderGrid/Source/RenderGridDeveloper/Private/Blueprints/RenderGridBlueprint.cpp b/Engine/Plugins/Experimental/RenderGrid/Source/RenderGridDeveloper/Private/Blueprints/RenderGridBlueprint.cpp
new file mode 100644
--- /dev/null
+++ b/Engine/Plugins/Experimental/RenderGrid/Source/RenderGridDeveloper/Private/Blueprints/RenderGridBlueprint.cpp
@@ -0,0 +1,160 @@
+// Copyright Epic Games, Inc. All Rights Reserved.
+
+#include "Blueprints/RenderGridBlueprint.h"
+
+#include "EdGraphSchema_K2_Actions.h"
+#include "Kismet2/BlueprintEditorUtils.h"
+#include "Kismet2/KismetEditorUtilities.h"
+#include "RenderGrid/RenderGrid.h"
+#include "RenderGrid/RenderGridBlueprintGeneratedClass.h"
+
+
+UClass* URenderGridBlueprint::GetBlueprintClass() const
+{
+	return URenderGridBlueprintGeneratedClass::StaticClass();
+}
+
+void URenderGridBlueprint::PostLoad()
+{
+	Super::PostLoad();
+
+	if (UbergraphPages.IsEmpty() || ((UbergraphPages.Num() == 1) && UbergraphPages[0]->Nodes.IsEmpty()))
+	{
+		if (!UbergraphPages.IsEmpty())
+		{
+			for (const TObjectPtr<UEdGraph>& Graph : UbergraphPages)
+			{
+				Graph->MarkAsGarbage();
+				Graph->Rename(nullptr, GetTransientPackage(), REN_ForceNoResetLoaders);
+			}
+			UbergraphPages.Empty();
+		}
+
+		UEdGraph* NewGraph = FBlueprintEditorUtils::CreateNewGraph(this, UEdGraphSchema_K2::GN_EventGraph, UEdGraph::StaticClass(), UEdGraphSchema_K2::StaticClass());
+		NewGraph->bAllowDeletion = false;
+
+		{// create every RenderGrid blueprint event >>
+			int32 i = 0;
+			for (const FString& Event : URenderGrid::GetBlueprintImplementableEvents())
+			{
+				int32 InOutNodePosY = (i * 256) - 48;
+				FKismetEditorUtilities::AddDefaultEventNode(this, NewGraph, FName(Event), URenderGrid::StaticClass(), InOutNodePosY);
+				i++;
+			}
+		}// create every RenderGrid blueprint event <<
+
+		FBlueprintEditorUtils::MarkBlueprintAsStructurallyModified(this);
+		FBlueprintEditorUtils::AddUbergraphPage(this, NewGraph);
+		LastEditedDocuments.AddUnique(NewGraph);
+	}
+
+	FCoreUObjectDelegates::OnObjectModified.RemoveAll(this);
+	OnChanged().RemoveAll(this);
+	FCoreUObjectDelegates::OnObjectModified.AddUObject(this, &URenderGridBlueprint::OnPreVariablesChange);
+	OnChanged().AddUObject(this, &URenderGridBlueprint::OnPostVariablesChange);
+
+	OnPreVariablesChange(this);
+	OnPostVariablesChange(this);
+}
+
+void URenderGridBlueprint::OnPreVariablesChange(UObject* InObject)
+{
+	if (InObject != this)
+	{
+		return;
+	}
+	LastNewVariables = NewVariables;
+}
+
+void URenderGridBlueprint::OnPostVariablesChange(UBlueprint* InBlueprint)
+{
+	if (InBlueprint != this)
+	{
+		return;
+	}
+
+	bool bFoundChange = false;
+
+	TMap<FGuid, int32> NewVariablesByGuid;
+	for (int32 VarIndex = 0; VarIndex < NewVariables.Num(); VarIndex++)
+	{
+		NewVariablesByGuid.Add(NewVariables[VarIndex].VarGuid, VarIndex);
+	}
+
+	TMap<FGuid, int32> OldVariablesByGuid;
+	for (int32 VarIndex = 0; VarIndex < LastNewVariables.Num(); VarIndex++)
+	{
+		OldVariablesByGuid.Add(LastNewVariables[VarIndex].VarGuid, VarIndex);
+	}
+
+	for (FBPVariableDescription& OldVariable : LastNewVariables)
+	{
+		if (!NewVariablesByGuid.Contains(OldVariable.VarGuid))
+		{
+			bFoundChange = true;
+			OnVariableRemoved(OldVariable);
+		}
+	}
+
+	for (FBPVariableDescription& NewVariable : NewVariables)
+	{
+		if (!OldVariablesByGuid.Contains(NewVariable.VarGuid))
+		{
+			bFoundChange = true;
+			OnVariableAdded(NewVariable);
+			continue;
+		}
+
+		int32 OldVarIndex = OldVariablesByGuid.FindChecked(NewVariable.VarGuid);
+		const FBPVariableDescription& OldVariable = LastNewVariables[OldVarIndex];
+		if (OldVariable.VarName != NewVariable.VarName)
+		{
+			bFoundChange = true;
+			OnVariableRenamed(NewVariable, OldVariable.VarName, NewVariable.VarName);
+		}
+
+		if (OldVariable.VarType != NewVariable.VarType)
+		{
+			bFoundChange = true;
+			OnVariableTypeChanged(NewVariable, OldVariable.VarType, NewVariable.VarType);
+		}
+
+		if (OldVariable.PropertyFlags != NewVariable.PropertyFlags)
+		{
+			bFoundChange = true;
+			OnVariablePropertyFlagsChanged(NewVariable, OldVariable.PropertyFlags, NewVariable.PropertyFlags);
+		}
+	}
+
+	LastNewVariables = NewVariables;
+
+	if (bFoundChange)
+	{
+		FBlueprintEditorUtils::MarkBlueprintAsStructurallyModified(this);
+	}
+}
+
+void URenderGridBlueprint::OnVariableAdded(FBPVariableDescription& InVar)
+{
+	MakeVariableTransientUnlessInstanceEditable(InVar);
+}
+
+void URenderGridBlueprint::OnVariablePropertyFlagsChanged(FBPVariableDescription& InVar, const uint64 InOldVarPropertyFlags, const uint64 InNewVarPropertyFlags)
+{
+	if ((InOldVarPropertyFlags & CPF_DisableEditOnInstance) != (InNewVarPropertyFlags & CPF_DisableEditOnInstance))// if the value of [Instance Editable] changed
+	{
+		MakeVariableTransientUnlessInstanceEditable(InVar);
+	}
+}
+
+void URenderGridBlueprint::MakeVariableTransientUnlessInstanceEditable(FBPVariableDescription& InVar)
+{
+	if ((InVar.PropertyFlags & CPF_DisableEditOnInstance) == 0)// if [Instance Editable]
+	{
+		InVar.PropertyFlags &= ~CPF_Transient;// set [Transient] to false
+	}
+	else
+	{
+		InVar.PropertyFlags |= CPF_Transient;// set [Transient] to true
+	}
+}
diff --git a/Engine/Plugins/Experimental/RenderGrid/Source/RenderGridDeveloper/Public/Blueprints/RenderGridBlueprint.h b/Engine/Plugins/Experimental/RenderGrid/Source/RenderGridDeveloper/Public/Blueprints/RenderGridBlueprint.h
new file mode 100644
--- /dev/null
+++ b/Engine/Plugins/Experimental/RenderGrid/Source/RenderGridDeveloper/Public/Blueprints/RenderGridBlueprint.h
@@ -0,0 +1,52 @@
+// Copyright Epic Games, Inc. All Rights Reserved.
+
+#pragma once
+
+#include "CoreMinimal.h"
+#include "Engine/Blueprint.h"
+#include "RenderGridBlueprint.generated.h"
+
+
+/**
+ * A UBlueprint child class for the RenderGrid modules.
+ *
+ * Required in order for a RenderGrid to be able to have a blueprint graph.
+ */
+UCLASS(BlueprintType, Meta = (IgnoreClassThumbnail))
+class RENDERGRIDDEVELOPER_API URenderGridBlueprint : public UBlueprint
+{
+	GENERATED_BODY()
+
+public:
+	//~ Begin UBlueprint Interface
+	virtual bool SupportedByDefaultBlueprintFactory() const override { return false; }
+	virtual bool IsValidForBytecodeOnlyRecompile() const override { return false; }
+	virtual bool SupportsGlobalVariables() const override { return true; }
+	virtual bool SupportsLocalVariables() const override { return true; }
+	virtual bool SupportsFunctions() const override { return true; }
+	virtual bool SupportsMacros() const override { return true; }
+	virtual bool SupportsDelegates() const override { return true; }
+	virtual bool SupportsEventGraphs() const override { return true; }
+	virtual bool SupportsAnimLayers() const override { return false; }
+
+	virtual UClass* GetBlueprintClass() const override;
+	virtual void PostLoad() override;
+	//~ End UBlueprint Interface
+
+private:
+	virtual void OnPreVariablesChange(UObject* InObject);
+	virtual void OnPostVariablesChange(UBlueprint* InBlueprint);
+
+protected:
+	virtual void OnVariableAdded(FBPVariableDescription& InVar);
+	virtual void OnVariableRemoved(FBPVariableDescription& InVar) {}
+	virtual void OnVariableRenamed(FBPVariableDescription& InVar, const FName& InOldVarName, const FName& InNewVarName) {}
+	virtual void OnVariableTypeChanged(FBPVariableDescription& InVar, const FEdGraphPinType& InOldVarType, const FEdGraphPinType& InNewVarType) {}
+	virtual void OnVariablePropertyFlagsChanged(FBPVariableDescription& InVar, const uint64 InOldVarPropertyFlags, const uint64 InNewVarPropertyFlags);
+
+protected:
+	virtual void MakeVariableTransientUnlessInstanceEditable(FBPVariableDescription& InVar);
+
+private:
+	TArray<FBPVariableDescription> LastNewVariables;
+};
diff --git a/Engine/Plugins/Experimental/RenderGrid/Source/RenderGridEditor/Private/AssetTypeActions/RenderGridBlueprintActions.cpp b/Engine/Plugins/Experimental/RenderGrid/Source/RenderGridEditor/Private/AssetTypeActions/RenderGridBlueprintActions.cpp
new file mode 100644
--- /dev/null
+++ b/Engine/Plugins/Experimental/RenderGrid/Source/RenderGridEditor/Private/AssetTypeActions/RenderGridBlueprintActions.cpp
@@ -0,0 +1,77 @@
+// Copyright Epic Games, Inc. All Rights Reserved.
+
+#include "RenderGridBlueprintActions.h"
+#include "Styling/AppStyle.h"
+#include "Blueprints/RenderGridBlueprint.h"
+#include "Factories/RenderGridFactory.h"
+#include "RenderGridEditorModule.h"
+#include "Styling/SlateIconFinder.h"
+
+#define LOCTEXT_NAMESPACE "AssetTypeActions"
+
+
+FText UE::RenderGrid::Private::FRenderGridBlueprintActions::GetName() const
+{
+	return LOCTEXT("AssetTypeActions_RenderGrid", "Render Grid");
+}
+
+FColor UE::RenderGrid::Private::FRenderGridBlueprintActions::GetTypeColor() const
+{
+	return FColor(255, 64, 64);
+}
+
+UClass* UE::RenderGrid::Private::FRenderGridBlueprintActions::GetSupportedClass() const
+{
+	return URenderGridBlueprint::StaticClass();
+}
+
+void UE::RenderGrid::Private::FRenderGridBlueprintActions::OpenAssetEditor(const TArray<UObject*>& InObjects, TSharedPtr<IToolkitHost> EditWithinLevelEditor)
+{
+	for (auto ObjIt = InObjects.CreateConstIterator(); ObjIt; ++ObjIt)
+	{
+		if (URenderGridBlueprint* RenderGridBlueprint = Cast<URenderGridBlueprint>(*ObjIt))
+		{
+			constexpr bool bBringToFrontIfOpen = true;
+			if (IAssetEditorInstance* EditorInstance = GEditor->GetEditorSubsystem<UAssetEditorSubsystem>()->FindEditorForAsset(RenderGridBlueprint, bBringToFrontIfOpen))
+			{
+				EditorInstance->FocusWindow(RenderGridBlueprint);
+			}
+			else
+			{
+				const EToolkitMode::Type Mode = EditWithinLevelEditor.IsValid() ? EToolkitMode::WorldCentric : EToolkitMode::Standalone;
+				IRenderGridEditorModule::Get().CreateRenderGridEditor(Mode, EditWithinLevelEditor, RenderGridBlueprint);
+			}
+		}
+	}
+}
+
+uint32 UE::RenderGrid::Private::FRenderGridBlueprintActions::GetCategories()
+{
+	return EAssetTypeCategories::Misc;
+}
+
+TSharedPtr<SWidget> UE::RenderGrid::Private::FRenderGridBlueprintActions::GetThumbnailOverlay(const FAssetData& AssetData) const
+{
+	const FSlateBrush* Icon = FSlateIconFinder::FindIconBrushForClass(URenderGridBlueprint::StaticClass());
+
+	return SNew(SBorder)
+		.BorderImage(FAppStyle::GetNoBrush())
+		.Visibility(EVisibility::HitTestInvisible)
+		.Padding(FMargin(0.0f, 0.0f, 0.0f, 3.0f))
+		.HAlign(HAlign_Right)
+		.VAlign(VAlign_Bottom)
+		[
+			SNew(SImage)
+			.Image(Icon)
+		];
+}
+
+UFactory* UE::RenderGrid::Private::FRenderGridBlueprintActions::GetFactoryForBlueprintType(UBlueprint* InBlueprint) const
+{
+	URenderGridBlueprintFactory* RenderGridBlueprintFactory = NewObject<URenderGridBlueprintFactory>();
+	RenderGridBlueprintFactory->ParentClass = TSubclassOf<URenderGrid>(*InBlueprint->GeneratedClass);
+	return RenderGridBlueprintFactory;
+}
+
+
+#undef LOCTEXT_NAMESPACE
