- Renamed the RenderPages plugin to RenderGrid

diff --git a/Engine/Plugins/Experimental/RenderGrid/Source/RenderGridEditor/Private/AssetTypeActions/RenderGridBlueprintActions.h b/Engine/Plugins/Experimental/RenderGrid/Source/RenderGridEditor/Private/AssetTypeActions/RenderGridBlueprintActions.h
new file mode 100644
--- /dev/null
+++ b/Engine/Plugins/Experimental/RenderGrid/Source/RenderGridEditor/Private/AssetTypeActions/RenderGridBlueprintActions.h
@@ -0,0 +1,30 @@
+// Copyright Epic Games, Inc. All Rights Reserved.
+
+#pragma once
+
+#include "AssetTypeActions_Base.h"
+#include "AssetTypeActions/AssetTypeActions_Blueprint.h"
+
+
+namespace UE::RenderGrid::Private
+{
+	/**
+	 * This class adds support for the RenderGrid (RenderGridBlueprint) class to the AssetTools module.
+	 */
+	class FRenderGridBlueprintActions : public FAssetTypeActions_Blueprint
+	{
+	public:
+		//~ Begin IAssetTypeActions Interface
+		virtual FText GetName() const override;
+		virtual FColor GetTypeColor() const override;
+		virtual UClass* GetSupportedClass() const override;
+		virtual void OpenAssetEditor(const TArray<UObject*>& InObjects, TSharedPtr<IToolkitHost> EditWithinLevelEditor = TSharedPtr<IToolkitHost>()) override;
+		virtual uint32 GetCategories() override;
+		virtual TSharedPtr<SWidget> GetThumbnailOverlay(const FAssetData& AssetData) const override;
+		//~ End IAssetTypeActions Interface
+
+		//~ Begin FAssetTypeActions_Blueprint interface
+		virtual UFactory* GetFactoryForBlueprintType(UBlueprint* InBlueprint) const override;
+		//~ End FAssetTypeActions_Blueprint interface
+	};
+}
diff --git a/Engine/Plugins/Experimental/RenderGrid/Source/RenderGridEditor/Private/Commands/RenderGridEditorCommands.cpp b/Engine/Plugins/Experimental/RenderGrid/Source/RenderGridEditor/Private/Commands/RenderGridEditorCommands.cpp
new file mode 100644
--- /dev/null
+++ b/Engine/Plugins/Experimental/RenderGrid/Source/RenderGridEditor/Private/Commands/RenderGridEditorCommands.cpp
@@ -0,0 +1,18 @@
+// Copyright Epic Games, Inc. All Rights Reserved.
+
+#include "Commands/RenderGridEditorCommands.h"
+
+#define LOCTEXT_NAMESPACE "FRenderGridEditor"
+
+
+void UE::RenderGrid::Private::FRenderGridEditorCommands::RegisterCommands()
+{
+	UI_COMMAND(AddJob, "Add", "Creates and adds a new job instance to the grid.", EUserInterfaceActionType::Button, FInputChord());
+	UI_COMMAND(DuplicateJob, "Duplicate", "Duplicates the selected job instances and adds them to the grid.", EUserInterfaceActionType::Button, FInputChord());
+	UI_COMMAND(DeleteJob, "Delete", "Deletes the selected job instances from the grid.", EUserInterfaceActionType::Button, FInputChord());
+
+	UI_COMMAND(BatchRenderList, "Render", "Renders the enabled job instances of the grid.", EUserInterfaceActionType::Button, FInputChord());
+}
+
+
+#undef LOCTEXT_NAMESPACE
diff --git a/Engine/Plugins/Experimental/RenderGrid/Source/RenderGridEditor/Private/Factories/RenderGridFactory.cpp b/Engine/Plugins/Experimental/RenderGrid/Source/RenderGridEditor/Private/Factories/RenderGridFactory.cpp
new file mode 100644
--- /dev/null
+++ b/Engine/Plugins/Experimental/RenderGrid/Source/RenderGridEditor/Private/Factories/RenderGridFactory.cpp
@@ -0,0 +1,64 @@
+// Copyright Epic Games, Inc. All Rights Reserved.
+
+#include "Factories/RenderGridFactory.h"
+#include "Blueprints/RenderGridBlueprint.h"
+#include "IContentBrowserSingleton.h"
+#include "Kismet2/BlueprintEditorUtils.h"
+#include "Kismet2/KismetEditorUtilities.h"
+#include "Misc/MessageDialog.h"
+#include "RenderGrid/RenderGrid.h"
+#include "RenderGrid/RenderGridBlueprintGeneratedClass.h"
+
+#define LOCTEXT_NAMESPACE "RenderGridBlueprintFactory"
+
+
+URenderGridBlueprintFactory::URenderGridBlueprintFactory()
+{
+	ParentClass = URenderGrid::StaticClass();
+	SupportedClass = URenderGridBlueprint::StaticClass();
+	bCreateNew = true; // This factory manufacture new objects from scratch
+	bEditAfterNew = true; // This factory will open the editor for each new object
+}
+
+UObject* URenderGridBlueprintFactory::FactoryCreateNew(UClass* InClass, UObject* InParent, FName InName, EObjectFlags Flags, UObject* Context, FFeedbackContext* Warn, FName CallingContext)
+{
+	// Make sure we are trying to factory a Render Grid Blueprint, then create and init one
+	check(InClass->IsChildOf(URenderGridBlueprint::StaticClass()));
+
+	if ((ParentClass == nullptr) || !FKismetEditorUtilities::CanCreateBlueprintOfClass(ParentClass) || !ParentClass->IsChildOf(URenderGrid::StaticClass()))
+	{
+		FFormatNamedArguments Args;
+		Args.Add(TEXT("ClassName"), (ParentClass != nullptr) ? FText::FromString(ParentClass->GetName()) : LOCTEXT("Null", "(null)"));
+		FMessageDialog::Open(EAppMsgType::Ok, FText::Format(LOCTEXT("CannotCreateRenderGridBlueprint", "Cannot create an Render Grid Blueprint based on the class '{0}'."), Args));
+		return nullptr;
+	}
+
+	if (URenderGridBlueprint* RenderGridBlueprint = CastChecked<URenderGridBlueprint>(FKismetEditorUtilities::CreateBlueprint(ParentClass, InParent, InName, BPTYPE_Normal, URenderGridBlueprint::StaticClass(), URenderGridBlueprintGeneratedClass::StaticClass(), CallingContext)))
+	{
+		RenderGridBlueprint->PostLoad();
+		return RenderGridBlueprint;
+	}
+	return nullptr;
+}
+
+UObject* URenderGridBlueprintFactory::FactoryCreateNew(UClass* InClass, UObject* InParent, FName InName, EObjectFlags Flags, UObject* Context, FFeedbackContext* Warn)
+{
+	return FactoryCreateNew(InClass, InParent, InName, Flags, Context, Warn, NAME_None);
+}
+
+bool URenderGridBlueprintFactory::ShouldShowInNewMenu() const
+{
+	return true;
+}
+
+uint32 URenderGridBlueprintFactory::GetMenuCategories() const
+{
+	//  if wanting to show it in its own category:
+	// IAssetTools& AssetTools = FModuleManager::LoadModuleChecked<FAssetToolsModule>("AssetTools").Get();
+	// return AssetTools.RegisterAdvancedAssetCategory("Render Grid", LOCTEXT("AssetCategoryName", "Render Grid"));
+
+	return EAssetTypeCategories::Misc;
+}
+
+
+#undef LOCTEXT_NAMESPACE
diff --git a/Engine/Plugins/Experimental/RenderGrid/Source/RenderGridEditor/Private/Factories/RenderGridFactory.h b/Engine/Plugins/Experimental/RenderGrid/Source/RenderGridEditor/Private/Factories/RenderGridFactory.h
new file mode 100644
--- /dev/null
+++ b/Engine/Plugins/Experimental/RenderGrid/Source/RenderGridEditor/Private/Factories/RenderGridFactory.h
@@ -0,0 +1,34 @@
+// Copyright Epic Games, Inc. All Rights Reserved.
+
+#pragma once
+
+#include "CoreMinimal.h"
+#include "Factories/Factory.h"
+#include "RenderGrid/RenderGridPropsSource.h"
+#include "RenderGridFactory.generated.h"
+
+
+class URenderGridBlueprint;
+class URenderGrid;
+
+
+/**
+ * The factory that creates URenderGridBlueprint (render grid) instances.
+ */
+UCLASS(MinimalAPI, HideCategories=Object)
+class URenderGridBlueprintFactory : public UFactory
+{
+	GENERATED_BODY()
+
+public:
+	URenderGridBlueprintFactory();
+
+	virtual UObject* FactoryCreateNew(UClass* InClass, UObject* InParent, FName InName, EObjectFlags Flags, UObject* Context, FFeedbackContext* Warn, FName CallingContext) override;
+	virtual UObject* FactoryCreateNew(UClass* InClass, UObject* InParent, FName InName, EObjectFlags Flags, UObject* Context, FFeedbackContext* Warn) override;
+	virtual bool ShouldShowInNewMenu() const override;
+	virtual uint32 GetMenuCategories() const override;
+
+	/** The parent class of the created blueprint. */
+	UPROPERTY(EditAnywhere, Category="Render Grid|Render Grid Factory", Meta = (AllowAbstract = ""))
+	TSubclassOf<URenderGrid> ParentClass;
+};
diff --git a/Engine/Plugins/Experimental/RenderGrid/Source/RenderGridEditor/Public/IRenderGridEditor.h b/Engine/Plugins/Experimental/RenderGrid/Source/RenderGridEditor/Public/IRenderGridEditor.h
new file mode 100644
--- /dev/null
+++ b/Engine/Plugins/Experimental/RenderGrid/Source/RenderGridEditor/Public/IRenderGridEditor.h
@@ -0,0 +1,98 @@
+// Copyright Epic Games, Inc. All Rights Reserved.
+
+#pragma once
+
+#include "CoreMinimal.h"
+#include "BlueprintEditor.h"
+
+
+class URenderGridMoviePipelineRenderJob;
+class URenderGridJob;
+class URenderGrid;
+class URenderGridBlueprint;
+
+namespace UE::RenderGrid::Private
+{
+	class FRenderGridBlueprintEditorToolbar;
+}
+
+
+namespace UE::RenderGrid
+{
+	/**
+	 * The render grid editor interface.
+	 */
+	class IRenderGridEditor : public FBlueprintEditor
+	{
+	public:
+		/** @return The render grid blueprint currently being edited in this editor. */
+		virtual URenderGridBlueprint* GetRenderGridBlueprint() const = 0;
+
+		/** @return The render grid instance. */
+		virtual URenderGrid* GetInstance() const = 0;
+
+		/** @return The render grid toolbar. */
+		virtual TSharedPtr<Private::FRenderGridBlueprintEditorToolbar> GetRenderGridToolbarBuilder() = 0;
+
+		/** Returns whether it is currently rendering or playing (so changes in the level and such should be ignored). */
+		virtual bool IsCurrentlyRenderingOrPlaying() const { return IsBatchRendering() || IsPreviewRendering() || IsValid(GEditor->PlayWorld); }
+
+		/** Returns whether it can currently render (like a preview render or a batch render) or not. */
+		virtual bool CanCurrentlyRender() const { return !IsCurrentlyRenderingOrPlaying(); }
+
+		/** Returns whether it is currently batch rendering or not. */
+		virtual bool IsBatchRendering() const = 0;
+
+		/** Returns the current batch rendering job, or null if it's not currently batch rendering. */
+		virtual URenderGridMoviePipelineRenderJob* GetBatchRenderJob() const = 0;
+
+		/** Returns whether it is currently preview rendering or not. */
+		virtual bool IsPreviewRendering() const = 0;
+
+		/** Returns the current preview rendering job, or null if it's not currently rendering a preview. */
+		virtual URenderGridMoviePipelineRenderJob* GetPreviewRenderJob() const = 0;
+
+		/** Sets the current preview rendering job, set it to null if it's not currently rendering a preview. */
+		virtual void SetPreviewRenderJob(URenderGridMoviePipelineRenderJob* Job) = 0;
+
+		/** Marks the editing asset as modified. */
+		virtual void MarkAsModified() = 0;
+
+		/** Get the currently selected render grid jobs. */
+		virtual TArray<URenderGridJob*> GetSelectedRenderGridJobs() const = 0;
+
+		/** Set the selected render grid jobs. */
+		virtual void SetSelectedRenderGridJobs(const TArray<URenderGridJob*>& Jobs) = 0;
+
+		DECLARE_MULTICAST_DELEGATE(FOnRenderGridChanged);
+		virtual FOnRenderGridChanged& OnRenderGridChanged() { return OnRenderGridChangedDelegate; }
+
+		DECLARE_MULTICAST_DELEGATE_OneParam(FOnRenderGridJobCreated, URenderGridJob*);
+		virtual FOnRenderGridJobCreated& OnRenderGridJobCreated() { return OnRenderGridJobCreatedDelegate; }
+
+		DECLARE_MULTICAST_DELEGATE(FOnRenderGridJobsSelectionChanged);
+		virtual FOnRenderGridJobsSelectionChanged& OnRenderGridJobsSelectionChanged() { return OnRenderGridJobsSelectionChangedDelegate; }
+
+		DECLARE_MULTICAST_DELEGATE_OneParam(FOnRenderGridBatchRenderingStarted, URenderGridMoviePipelineRenderJob*);
+		virtual FOnRenderGridBatchRenderingStarted& OnRenderGridBatchRenderingStarted() { return OnRenderGridBatchRenderingStartedDelegate; }
+
+		DECLARE_MULTICAST_DELEGATE_OneParam(FOnRenderGridBatchRenderingFinished, URenderGridMoviePipelineRenderJob*);
+		virtual FOnRenderGridBatchRenderingFinished& OnRenderGridBatchRenderingFinished() { return OnRenderGridBatchRenderingFinishedDelegate; }
+
+	private:
+		/** The delegate for when data in the render grid changed. */
+		FOnRenderGridChanged OnRenderGridChangedDelegate;
+
+		/** The delegate for when a render grid job is created. */
+		FOnRenderGridJobCreated OnRenderGridJobCreatedDelegate;
+
+		/** The delegate for when the selection of render grid jobs changed. */
+		FOnRenderGridJobsSelectionChanged OnRenderGridJobsSelectionChangedDelegate;
+
+		/** The delegate for when batch rendering of a render grid has started. */
+		FOnRenderGridBatchRenderingStarted OnRenderGridBatchRenderingStartedDelegate;
+
+		/** The delegate for when batch rendering of a render grid has finished, successfully or not. */
+		FOnRenderGridBatchRenderingFinished OnRenderGridBatchRenderingFinishedDelegate;
+	};
+}
