- Implemented a color wheel fix (still partially bugging, still a work in progress).
- Removed the local props type.
- Hid the props type dropdown.
- Made the props type be automatically converted to remote control if a render grid asset has it set to local.
- Added a on-open and on-exit blueprint event (called BeginEditor and EndEditor, similar to BeginPlay and EndPlay).
- Made it so that making changes in a remote control preset cause the render grid UI to update correctly now.
- Made the preview viewer mode wait with rendering when the color wheel is open (until it is closed again).
- Made the color wheel render-prevention fix only be applied to the preview frame viewer (and not to the rendered viewer).
- Made the header checkbox (for job enabled state) only affect the jobs that are visible in the list (the ones that aren't filtered out by the search bar).

diff --git a/Engine/Plugins/Experimental/RenderGrid/Source/RenderGrid/Private/RenderGrid/RenderGrid.cpp b/Engine/Plugins/Experimental/RenderGrid/Source/RenderGrid/Private/RenderGrid/RenderGrid.cpp
--- a/Engine/Plugins/Experimental/RenderGrid/Source/RenderGrid/Private/RenderGrid/RenderGrid.cpp
+++ b/Engine/Plugins/Experimental/RenderGrid/Source/RenderGrid/Private/RenderGrid/RenderGrid.cpp
@@ -461,7 +461,7 @@ bool URenderGridJob::SetRemoteControlValue(const TSharedPtr<FRemoteControlEntity
 
 URenderGrid::URenderGrid()
 	: Guid(FGuid::NewGuid())
-	, PropsSourceType(ERenderGridPropsSourceType::Local)
+	, PropsSourceType(ERenderGridPropsSourceType::RemoteControl)
 	, PropsSourceOrigin_RemoteControl(nullptr)
 	, bExecutingBlueprintEvent(false)
 	, CachedPropsSource(nullptr)
@@ -527,6 +527,25 @@ void URenderGrid::PostLoad()
 {
 	Super::PostLoad();
 	LoadValuesFromCDO();
+
+	if (PropsSourceType == ERenderGridPropsSourceType::Local)
+	{
+		PropsSourceType = ERenderGridPropsSourceType::RemoteControl;
+	}
+}
+
+void URenderGrid::BeginEditor()
+{
+	bExecutingBlueprintEvent = true;
+	ReceiveBeginEditor();
+	bExecutingBlueprintEvent = false;
+}
+
+void URenderGrid::EndEditor()
+{
+	bExecutingBlueprintEvent = true;
+	ReceiveEndEditor();
+	bExecutingBlueprintEvent = false;
 }
 
 void URenderGrid::BeginBatchRender(URenderGridQueue* Queue)
@@ -622,14 +641,11 @@ void URenderGrid::SetPropsSource(ERenderGridPropsSourceType InPropsSourceType, U
 {
 	if (InPropsSourceType == ERenderGridPropsSourceType::RemoteControl)
 	{
-		if (URemoteControlPreset* InPropsSourceOrigin_RemoteControl = Cast<URemoteControlPreset>(InPropsSourceOrigin))
-		{
-			PropsSourceType = InPropsSourceType;
-			PropsSourceOrigin_RemoteControl = InPropsSourceOrigin_RemoteControl;
-			return;
-		}
+		PropsSourceType = InPropsSourceType;
+		PropsSourceOrigin_RemoteControl = Cast<URemoteControlPreset>(InPropsSourceOrigin);
+		return;
 	}
-	PropsSourceType = ERenderGridPropsSourceType::Local;
+	PropsSourceType = ERenderGridPropsSourceType::RemoteControl;
 }
 
 URenderGridPropsSourceBase* URenderGrid::GetPropsSource() const
diff --git a/Engine/Plugins/Experimental/RenderGrid/Source/RenderGrid/Private/RenderGrid/RenderGridPropsSource.cpp b/Engine/Plugins/Experimental/RenderGrid/Source/RenderGrid/Private/RenderGrid/RenderGridPropsSource.cpp
--- a/Engine/Plugins/Experimental/RenderGrid/Source/RenderGrid/Private/RenderGrid/RenderGridPropsSource.cpp
+++ b/Engine/Plugins/Experimental/RenderGrid/Source/RenderGrid/Private/RenderGrid/RenderGridPropsSource.cpp
@@ -13,17 +13,27 @@ bool /*URenderGridPropRemoteControl::*/GetObjectRef(const TSharedPtr<FRemoteCont
 	{
 		return false;
 	}
-	if (UObject* FieldBoundObject = Field->GetBoundObject(); IsValid(FieldBoundObject))
+
+	if (const UScriptStruct* PropertyStruct = Field->GetStruct())
 	{
-		FRCObjectReference ObjectRef;
-		FString* ErrorText = nullptr;
-		if (IRemoteControlModule::Get().ResolveObjectProperty(Access, FieldBoundObject, Field->FieldPathInfo, ObjectRef, ErrorText))
+		if (PropertyStruct->IsChildOf(FRemoteControlProperty::StaticStruct()))
 		{
-			OutObjectRef = ObjectRef;
-			return true;
+			if (UObject* FieldBoundObject = Field->GetBoundObject(); IsValid(FieldBoundObject))
+			{
+				FRCObjectReference ObjectRef;
+				FString* ErrorText = nullptr;
+
+				if (IRemoteControlModule::Get().ResolveObjectProperty(Access, FieldBoundObject, Field->FieldPathInfo, ObjectRef, ErrorText))
+				{
+					OutObjectRef = ObjectRef;
+					return true;
+				}
+
+				UE_LOG(LogRenderGrid, Warning, TEXT("Couldn\'t resolve object property \"%s\" in object \"%s\": %s"), *Field->FieldName.ToString(), *FieldBoundObject->GetPathName(), (ErrorText ? **ErrorText : TEXT("unknown")));
+			}
 		}
-		UE_LOG(LogRenderGrid, Warning, TEXT("Couldn\'t resolve object property \"%s\" in object \"%s\": %s"), *Field->FieldName.ToString(), *FieldBoundObject->GetPathName(), (ErrorText ? **ErrorText : TEXT("unknown")));
 	}
+
 	return false;
 }
 
diff --git a/Engine/Plugins/Experimental/RenderGrid/Source/RenderGrid/Public/RenderGrid/RenderGrid.h b/Engine/Plugins/Experimental/RenderGrid/Source/RenderGrid/Public/RenderGrid/RenderGrid.h
--- a/Engine/Plugins/Experimental/RenderGrid/Source/RenderGrid/Public/RenderGrid/RenderGrid.h
+++ b/Engine/Plugins/Experimental/RenderGrid/Source/RenderGrid/Public/RenderGrid/RenderGrid.h
@@ -301,6 +301,8 @@ public:
 	static TArray<FString> GetBlueprintImplementableEvents()
 	{
 		return {
+			TEXT("ReceiveBeginEditor"),
+			TEXT("ReceiveEndEditor"),
 			TEXT("ReceiveBeginBatchRender"),
 			TEXT("ReceiveEndBatchRender"),
 			TEXT("ReceiveBeginJobRender"),
@@ -311,6 +313,22 @@ public:
 	}
 
 protected:
+	/**
+	 * Event for when this asset is opened in the editor. The asset will also be closed and reopened during a blueprint compile.
+	 * 
+	 * In here, you could for example obtain jobs from an external source and replace the current jobs of this render grid asset with those.
+	 */
+	UFUNCTION(BlueprintImplementableEvent, Meta=(DisplayName="BeginEditor"))
+	void ReceiveBeginEditor();
+
+	/**
+	 * Event for when this asset is closed in the editor. The asset will also be closed and reopened during a blueprint compile.
+	 * 
+	 * In here, you could do any cleanup required at the end of editing this render grid asset.
+	 */
+	UFUNCTION(BlueprintImplementableEvent, Meta=(DisplayName="EndEditor"))
+	void ReceiveEndEditor();
+
 	/**
 	 * Event for when batch rendering begins.
 	 * 
@@ -366,6 +384,12 @@ protected:
 	void ReceiveEndViewportRender(URenderGridJob* Job);
 
 public:
+	/** Overridable native event for when this asset is opened in the editor. The asset will also be closed and reopened during a blueprint compile. */
+	virtual void BeginEditor();
+
+	/** Overridable native event for when this asset is closed in the editor. The asset will also be closed and reopened during a blueprint compile. */
+	virtual void EndEditor();
+
 	/** Overridable native event for when batch rendering begins. */
 	virtual void BeginBatchRender(URenderGridQueue* Queue);
 
@@ -515,7 +539,7 @@ private:
 
 
 	/** The type of the properties that a job in this grid can have. */
-	UPROPERTY(EditInstanceOnly, Category="Render Grid|Grid", Meta=(DisplayName="Properties Type", AllowPrivateAccess="true"))
+	UPROPERTY(/*EditInstanceOnly, Category="Render Grid|Grid", Meta=(DisplayName="Properties Type", AllowPrivateAccess="true")*/)
 	ERenderGridPropsSourceType PropsSourceType;
 
 	/** The remote control properties that a job in this grid can have, only use this if PropsSourceType is ERenderGridPropsSourceType::RemoteControl. */
diff --git a/Engine/Plugins/Experimental/RenderGrid/Source/RenderGrid/Public/RenderGrid/RenderGridPropsSource.h b/Engine/Plugins/Experimental/RenderGrid/Source/RenderGrid/Public/RenderGrid/RenderGridPropsSource.h
--- a/Engine/Plugins/Experimental/RenderGrid/Source/RenderGrid/Public/RenderGrid/RenderGridPropsSource.h
+++ b/Engine/Plugins/Experimental/RenderGrid/Source/RenderGrid/Public/RenderGrid/RenderGridPropsSource.h
@@ -17,7 +17,7 @@
 UENUM(BlueprintType)
 enum class ERenderGridPropsSourceType : uint8
 {
-	Local = 0 UMETA(DisplayName = "Local Source"),
+	Local = 0 UMETA(Hidden, DisplayName = "Local Source"),
 	RemoteControl = 1 UMETA(DisplayName = "Remote Control Preset")
 };
 
diff --git a/Engine/Plugins/Experimental/RenderGrid/Source/RenderGridEditor/Private/Toolkit/RenderGridEditor.cpp b/Engine/Plugins/Experimental/RenderGrid/Source/RenderGridEditor/Private/Toolkit/RenderGridEditor.cpp
--- a/Engine/Plugins/Experimental/RenderGrid/Source/RenderGridEditor/Private/Toolkit/RenderGridEditor.cpp
+++ b/Engine/Plugins/Experimental/RenderGrid/Source/RenderGridEditor/Private/Toolkit/RenderGridEditor.cpp
@@ -90,6 +90,7 @@ UE::RenderGrid::Private::FRenderGridEditor::FRenderGridEditor()
 	, bRunRenderNewBatch(false)
 	, BatchRenderQueue(nullptr)
 	, PreviewRenderQueue(nullptr)
+	, DebuggingTimeInSecondsRemaining(0)
 	, bIsDebugging(false)
 {}
 
@@ -581,7 +582,13 @@ void UE::RenderGrid::Private::FRenderGridEditor::DestroyInstance()
 {
 	if (URenderGrid* Instance = GetInstance(); IsValid(Instance))
 	{
+		// Execute the blueprint event
+		Instance->EndEditor();
+
+		// Store the data
 		Instance->OnClose();
+
+		// Garbage collection
 		RenderGridWeakPtr.Reset();
 		Instance->MarkAsGarbage();
 	}
@@ -617,6 +624,9 @@ void UE::RenderGrid::Private::FRenderGridEditor::UpdateInstance(UBlueprint* InBl
 		// Broadcast the events
 		OnRenderGridChanged().Broadcast();
 		OnRenderGridJobsSelectionChanged().Broadcast();
+
+		// Execute the blueprint event
+		RenderGrid->BeginEditor();
 	}
 }
 
diff --git a/Engine/Plugins/Experimental/RenderGrid/Source/RenderGridEditor/Private/UI/Components/SRenderGridRemoteControlField.cpp b/Engine/Plugins/Experimental/RenderGrid/Source/RenderGridEditor/Private/UI/Components/SRenderGridRemoteControlField.cpp
--- a/Engine/Plugins/Experimental/RenderGrid/Source/RenderGridEditor/Private/UI/Components/SRenderGridRemoteControlField.cpp
+++ b/Engine/Plugins/Experimental/RenderGrid/Source/RenderGridEditor/Private/UI/Components/SRenderGridRemoteControlField.cpp
@@ -111,38 +111,6 @@ void UE::RenderGrid::Private::SRenderGridRemoteControlField::RefreshValue()
 	if (!Generator.IsValid() || (GetFieldType() != EExposedFieldType::Property))
 	{
 		Refresh();
-		return;
-	}
-
-	if (const TSharedPtr<FRemoteControlField> Field = FieldWeakPtr.Pin())
-	{
-		TArray<UObject*> Objects = Field->GetBoundObjects();
-		if (Objects.Num() <= 0)
-		{
-			Generator->SetObjects({});
-			ChildSlot.AttachWidget(MakeFieldWidget(SNullWidget::NullWidget));
-			return;
-		}
-
-		Generator->SetObjects({Objects[0]});
-		if (TSharedPtr<IDetailTreeNode> Node = RenderGridWidgetUtils::FindNode(Generator->GetRootTreeNodes(), Field->FieldPathInfo.ToPathPropertyString(), RenderGridWidgetUtils::ERenderGridFindNodeMethod::Path))
-		{
-			TArray<TSharedRef<IDetailTreeNode>> ChildNodes;
-			Node->GetChildren(ChildNodes);
-
-			ChildWidgets.Reset(ChildNodes.Num());
-			for (const TSharedRef<IDetailTreeNode>& ChildNode : ChildNodes)
-			{
-				ChildWidgets.Add(SNew(SRenderGridRemoteControlFieldChildNode, ChildNode, ColumnSizeData));
-			}
-
-			//TODO:  still causes the value widgets (like the color wheel) to disconnect when this function is called,  have to somehow prevent this disconnection from happening,
-			//       skipping this line all-together causes the value widgets to not respond to user input anymore,  same happens when trying to cache and reuse the output of Node->CreateNodeWidgets()
-			ChildSlot.AttachWidget(MakeFieldWidget(ExposedFieldUtils::CreateNodeValueWidget(Node->CreateNodeWidgets())));
-			return;
-		}
-
-		ChildSlot.AttachWidget(MakeFieldWidget(SNullWidget::NullWidget));
 	}
 }
 
diff --git a/Engine/Plugins/Experimental/RenderGrid/Source/RenderGridEditor/Private/UI/Components/SRenderGridRemoteControlTreeNode.h b/Engine/Plugins/Experimental/RenderGrid/Source/RenderGridEditor/Private/UI/Components/SRenderGridRemoteControlTreeNode.h
--- a/Engine/Plugins/Experimental/RenderGrid/Source/RenderGridEditor/Private/UI/Components/SRenderGridRemoteControlTreeNode.h
+++ b/Engine/Plugins/Experimental/RenderGrid/Source/RenderGridEditor/Private/UI/Components/SRenderGridRemoteControlTreeNode.h
@@ -45,11 +45,12 @@ namespace UE::RenderGrid::Private
 		TObjectPtr<URemoteControlPreset> Preset = nullptr;
 		FRenderGridRemoteControlColumnSizeData ColumnSizeData;
 		TSharedPtr<FRemoteControlEntity> Entity = nullptr;
+		FName EntityFieldLabel;
 
 		bool operator!=(const FRenderGridRemoteControlGenerateWidgetArgs& RHS) const { return !(*this == RHS); }
 		bool operator==(const FRenderGridRemoteControlGenerateWidgetArgs& RHS) const
 		{
-			return (Preset == RHS.Preset) && (ColumnSizeData == RHS.ColumnSizeData) && (Entity == RHS.Entity);
+			return (Preset == RHS.Preset) && (ColumnSizeData == RHS.ColumnSizeData) && (Entity == RHS.Entity) && (EntityFieldLabel == RHS.EntityFieldLabel);
 		}
 	};
 
diff --git a/Engine/Plugins/Experimental/RenderGrid/Source/RenderGridEditor/Private/UI/SRenderGridJobList.cpp b/Engine/Plugins/Experimental/RenderGrid/Source/RenderGridEditor/Private/UI/SRenderGridJobList.cpp
--- a/Engine/Plugins/Experimental/RenderGrid/Source/RenderGridEditor/Private/UI/SRenderGridJobList.cpp
+++ b/Engine/Plugins/Experimental/RenderGrid/Source/RenderGridEditor/Private/UI/SRenderGridJobList.cpp
@@ -189,47 +189,35 @@ void UE::RenderGrid::Private::SRenderGridJobList::OnRenderGridJobCreated(URender
 
 void UE::RenderGrid::Private::SRenderGridJobList::OnHeaderCheckboxToggled(ECheckBoxState State)
 {
-	if (const TSharedPtr<IRenderGridEditor> BlueprintEditor = BlueprintEditorWeakPtr.Pin())
-	{
-		if (URenderGrid* Grid = BlueprintEditor->GetInstance(); IsValid(Grid))
-		{
-			bool bRefresh = false;
+	bool bRefresh = false;
 
-			for (URenderGridJob* Job : Grid->GetRenderGridJobs())
-			{
-				Job->SetIsEnabled(State == ECheckBoxState::Checked);
-				bRefresh = true;
-			}
+	for (URenderGridJob* Job : RenderGridJobs)
+	{
+		Job->SetIsEnabled(State == ECheckBoxState::Checked);
+		bRefresh = true;
+	}
 
-			if (bRefresh)
-			{
-				Refresh();
-			}
-		}
+	if (bRefresh)
+	{
+		Refresh();
 	}
 }
 
 ECheckBoxState UE::RenderGrid::Private::SRenderGridJobList::GetDesiredHeaderEnabledCheckboxState()
 {
 	ECheckBoxState State = ECheckBoxState::Checked;
-	if (const TSharedPtr<IRenderGridEditor> BlueprintEditor = BlueprintEditorWeakPtr.Pin())
+	bool bFirstJob = true;
+	for (URenderGridJob* Job : RenderGridJobs)
 	{
-		if (URenderGrid* Grid = BlueprintEditor->GetInstance(); IsValid(Grid))
+		ECheckBoxState JobState = (Job->GetIsEnabled() ? ECheckBoxState::Checked : ECheckBoxState::Unchecked);
+		if (bFirstJob)
 		{
-			bool bFirstJob = true;
-			for (URenderGridJob* Job : Grid->GetRenderGridJobs())
-			{
-				ECheckBoxState JobState = (Job->GetIsEnabled() ? ECheckBoxState::Checked : ECheckBoxState::Unchecked);
-				if (bFirstJob)
-				{
-					State = JobState;
-					bFirstJob = false;
-				}
-				else if (State != JobState)
-				{
-					return ECheckBoxState::Undetermined;
-				}
-			}
+			State = JobState;
+			bFirstJob = false;
+		}
+		else if (State != JobState)
+		{
+			return ECheckBoxState::Undetermined;
 		}
 	}
 	return State;
diff --git a/Engine/Plugins/Experimental/RenderGrid/Source/RenderGridEditor/Private/UI/SRenderGridProps.cpp b/Engine/Plugins/Experimental/RenderGrid/Source/RenderGridEditor/Private/UI/SRenderGridProps.cpp
--- a/Engine/Plugins/Experimental/RenderGrid/Source/RenderGridEditor/Private/UI/SRenderGridProps.cpp
+++ b/Engine/Plugins/Experimental/RenderGrid/Source/RenderGridEditor/Private/UI/SRenderGridProps.cpp
@@ -36,7 +36,8 @@ void UE::RenderGrid::Private::SRenderGridProps::Construct(const FArguments& InAr
 		.BorderImage(new FSlateNoResource());
 
 	Refresh();
-	InBlueprintEditor->OnRenderGridChanged().AddSP(this, &SRenderGridProps::Refresh);
+	//InBlueprintEditor->OnRenderGridChanged().AddSP(this, &SRenderGridProps::Refresh);
+	InBlueprintEditor->OnRenderGridJobsSelectionChanged().AddSP(this, &SRenderGridProps::Refresh);
 	InBlueprintEditor->OnRenderGridBatchRenderingStarted().AddSP(this, &SRenderGridProps::OnBatchRenderingStarted);
 	InBlueprintEditor->OnRenderGridBatchRenderingFinished().AddSP(this, &SRenderGridProps::OnBatchRenderingFinished);
 
diff --git a/Engine/Plugins/Experimental/RenderGrid/Source/RenderGridEditor/Private/UI/SRenderGridPropsRemoteControl.cpp b/Engine/Plugins/Experimental/RenderGrid/Source/RenderGridEditor/Private/UI/SRenderGridPropsRemoteControl.cpp
--- a/Engine/Plugins/Experimental/RenderGrid/Source/RenderGridEditor/Private/UI/SRenderGridPropsRemoteControl.cpp
+++ b/Engine/Plugins/Experimental/RenderGrid/Source/RenderGridEditor/Private/UI/SRenderGridPropsRemoteControl.cpp
@@ -20,9 +20,9 @@ void UE::RenderGrid::Private::SRenderGridPropsRemoteControl::Construct(const FAr
 	PropsSource = InPropsSource;
 
 	SAssignNew(RowWidgetsContainer, SVerticalBox);
-	UpdateStoredValuesAndRefresh();
+	UpdateStoredValuesAndRefresh(true);
 
-	InBlueprintEditor->OnRenderGridJobsSelectionChanged().AddSP(this, &SRenderGridPropsRemoteControl::Refresh);
+	InBlueprintEditor->OnRenderGridJobsSelectionChanged().AddSP(this, &SRenderGridPropsRemoteControl::OnRenderGridJobsSelectionChanged);
 	if (IsValid(PropsSource))
 	{
 		if (TObjectPtr<URemoteControlPreset> Preset = PropsSource->GetProps()->GetRemoteControlPreset())
@@ -30,6 +30,7 @@ void UE::RenderGrid::Private::SRenderGridPropsRemoteControl::Construct(const FAr
 			Preset->OnEntityExposed().AddSP(this, &SRenderGridPropsRemoteControl::OnRemoteControlEntitiesExposed);
 			Preset->OnEntityUnexposed().AddSP(this, &SRenderGridPropsRemoteControl::OnRemoteControlEntitiesUnexposed);
 			Preset->OnEntitiesUpdated().AddSP(this, &SRenderGridPropsRemoteControl::OnRemoteControlEntitiesUpdated);
+			Preset->OnPresetLayoutModified().AddSP(this, &SRenderGridPropsRemoteControl::OnRemoteControlPresetLayoutModified);
 			Preset->OnExposedPropertiesModified().AddSP(this, &SRenderGridPropsRemoteControl::OnRemoteControlExposedPropertiesModified);
 		}
 	}
@@ -48,16 +49,16 @@ void UE::RenderGrid::Private::SRenderGridPropsRemoteControl::Construct(const FAr
 END_SLATE_FUNCTION_BUILD_OPTIMIZATION
 
 
-void UE::RenderGrid::Private::SRenderGridPropsRemoteControl::UpdateStoredValuesAndRefresh()
+void UE::RenderGrid::Private::SRenderGridPropsRemoteControl::UpdateStoredValuesAndRefresh(const bool bForce)
 {
 	if (const TSharedPtr<IRenderGridEditor> BlueprintEditor = BlueprintEditorWeakPtr.Pin())
 	{
 		IRenderGridModule::Get().GetManager().UpdateRenderGridJobsPropValues(BlueprintEditor->GetInstance());
-		Refresh();
+		Refresh(bForce);
 	}
 }
 
-void UE::RenderGrid::Private::SRenderGridPropsRemoteControl::Refresh()
+void UE::RenderGrid::Private::SRenderGridPropsRemoteControl::Refresh(const bool bForce)
 {
 	if (!RowWidgetsContainer.IsValid())
 	{
@@ -92,9 +93,11 @@ void UE::RenderGrid::Private::SRenderGridPropsRemoteControl::Refresh()
 
 				if (URemoteControlPreset* Preset = Props->GetRemoteControlPreset(); IsValid(Preset))
 				{
+					TSharedPtr<FRemoteControlField> EntityField = StaticCastSharedPtr<FRemoteControlField>(Prop->GetRemoteControlEntity());
 					FRenderGridRemoteControlGenerateWidgetArgs Args;
 					Args.Preset = Preset;
 					Args.Entity = Prop->GetRemoteControlEntity();
+					Args.EntityFieldLabel = EntityField.IsValid() ? EntityField->GetLabel() : FName();
 					Args.ColumnSizeData.LeftColumnWidth = 0.3;
 					Args.ColumnSizeData.RightColumnWidth = 0.7;
 					NewRowWidgetsArgs.Add(Args);
@@ -102,7 +105,7 @@ void UE::RenderGrid::Private::SRenderGridPropsRemoteControl::Refresh()
 			}
 		}
 
-		if (RowWidgetsArgs != NewRowWidgetsArgs)
+		if (bForce || (RowWidgetsArgs != NewRowWidgetsArgs))
 		{
 			RowWidgetsArgs = NewRowWidgetsArgs;
 			RowWidgetsContainer->ClearChildren();
diff --git a/Engine/Plugins/Experimental/RenderGrid/Source/RenderGridEditor/Private/UI/SRenderGridPropsRemoteControl.h b/Engine/Plugins/Experimental/RenderGrid/Source/RenderGridEditor/Private/UI/SRenderGridPropsRemoteControl.h
--- a/Engine/Plugins/Experimental/RenderGrid/Source/RenderGridEditor/Private/UI/SRenderGridPropsRemoteControl.h
+++ b/Engine/Plugins/Experimental/RenderGrid/Source/RenderGridEditor/Private/UI/SRenderGridPropsRemoteControl.h
@@ -31,15 +31,17 @@ namespace UE::RenderGrid::Private
 		void Construct(const FArguments& InArgs, TSharedPtr<IRenderGridEditor> InBlueprintEditor, URenderGridPropsSourceRemoteControl* InPropsSource);
 
 		/** Obtain the latest prop values and refreshes the content of this widget. */
-		void UpdateStoredValuesAndRefresh();
+		void UpdateStoredValuesAndRefresh(const bool bForce = false);
 
 		/** Refreshes the content of this widget. */
-		void Refresh();
+		void Refresh(const bool bForce = false);
 
 	private:
-		void OnRemoteControlEntitiesExposed(URemoteControlPreset* Preset, const FGuid& EntityId) { UpdateStoredValuesAndRefresh(); }
-		void OnRemoteControlEntitiesUnexposed(URemoteControlPreset* Preset, const FGuid& EntityId) { UpdateStoredValuesAndRefresh(); }
+		void OnRenderGridJobsSelectionChanged() { Refresh(true); }
+		void OnRemoteControlEntitiesExposed(URemoteControlPreset* Preset, const FGuid& EntityId) { UpdateStoredValuesAndRefresh(true); }
+		void OnRemoteControlEntitiesUnexposed(URemoteControlPreset* Preset, const FGuid& EntityId) { UpdateStoredValuesAndRefresh(true); }
 		void OnRemoteControlEntitiesUpdated(URemoteControlPreset* Preset, const TSet<FGuid>& ModifiedEntities) { UpdateStoredValuesAndRefresh(); }
+		void OnRemoteControlPresetLayoutModified(URemoteControlPreset* Preset) { UpdateStoredValuesAndRefresh(true); }
 		void OnRemoteControlExposedPropertiesModified(URemoteControlPreset* Preset, const TSet<FGuid>& ModifiedProperties);
 
 	private:
diff --git a/Engine/Plugins/Experimental/RenderGrid/Source/RenderGridEditor/Private/UI/SRenderGridViewerPreview.cpp b/Engine/Plugins/Experimental/RenderGrid/Source/RenderGridEditor/Private/UI/SRenderGridViewerPreview.cpp
--- a/Engine/Plugins/Experimental/RenderGrid/Source/RenderGridEditor/Private/UI/SRenderGridViewerPreview.cpp
+++ b/Engine/Plugins/Experimental/RenderGrid/Source/RenderGridEditor/Private/UI/SRenderGridViewerPreview.cpp
@@ -13,6 +13,7 @@
 #include "Sections/MovieSceneSubSection.h"
 #include "Styles/RenderGridEditorStyle.h"
 #include "Tracks/MovieSceneSubTrack.h"
+#include "Widgets/Colors/SColorPicker.h"
 #include "Widgets/Input/SButton.h"
 #include "Widgets/Input/SSlider.h"
 #include "Widgets/Layout/SScaleBox.h"
@@ -304,7 +305,7 @@ void UE::RenderGrid::Private::SRenderGridViewerPreview::InternalRenderNewPreview
 
 void UE::RenderGrid::Private::SRenderGridViewerPreview::InternalRenderNewPreviewOfRenderGridJob(URenderGridJob* Job)
 {
-	if (GetTickSpaceGeometry().Size.X <= 0)
+	if ((GetTickSpaceGeometry().Size.X <= 0) || (IsPreviewWidget() && (GetColorPicker() != nullptr)))
 	{
 		// don't render, try again next frame
 		RenderNewPreview();
diff --git a/Engine/Plugins/Experimental/RenderGrid/Source/RenderGridEditor/RenderGridEditor.Build.cs b/Engine/Plugins/Experimental/RenderGrid/Source/RenderGridEditor/RenderGridEditor.Build.cs
--- a/Engine/Plugins/Experimental/RenderGrid/Source/RenderGridEditor/RenderGridEditor.Build.cs
+++ b/Engine/Plugins/Experimental/RenderGrid/Source/RenderGridEditor/RenderGridEditor.Build.cs
@@ -22,6 +22,7 @@ public class RenderGridEditor : ModuleRules
 		PrivateDependencyModuleNames.AddRange(
 			new string[]
 			{
+				"AppFramework",
 				"AssetTools",
 				"DesktopPlatform",
 				"EditorStyle",
