Added the new icon to the render grid plugin.

diff --git a/Engine/Plugins/Experimental/RenderGrid/Source/RenderGridEditor/Private/AssetTypeActions/RenderGridBlueprintActions.cpp b/Engine/Plugins/Experimental/RenderGrid/Source/RenderGridEditor/Private/AssetTypeActions/RenderGridBlueprintActions.cpp
--- a/Engine/Plugins/Experimental/RenderGrid/Source/RenderGridEditor/Private/AssetTypeActions/RenderGridBlueprintActions.cpp
+++ b/Engine/Plugins/Experimental/RenderGrid/Source/RenderGridEditor/Private/AssetTypeActions/RenderGridBlueprintActions.cpp
@@ -1,11 +1,9 @@
 // Copyright Epic Games, Inc. All Rights Reserved.
 
 #include "RenderGridBlueprintActions.h"
-#include "Styling/AppStyle.h"
 #include "Blueprints/RenderGridBlueprint.h"
 #include "Factories/RenderGridFactory.h"
 #include "RenderGridEditorModule.h"
-#include "Styling/SlateIconFinder.h"
 
 #define LOCTEXT_NAMESPACE "AssetTypeActions"
 
@@ -17,7 +15,7 @@ FText UE::RenderGrid::Private::FRenderGridBlueprintActions::GetName() const
 
 FColor UE::RenderGrid::Private::FRenderGridBlueprintActions::GetTypeColor() const
 {
-	return FColor(255, 64, 64);
+	return FColor(200, 80, 80);
 }
 
 UClass* UE::RenderGrid::Private::FRenderGridBlueprintActions::GetSupportedClass() const
@@ -50,22 +48,6 @@ uint32 UE::RenderGrid::Private::FRenderGridBlueprintActions::GetCategories()
 	return EAssetTypeCategories::Misc;
 }
 
-TSharedPtr<SWidget> UE::RenderGrid::Private::FRenderGridBlueprintActions::GetThumbnailOverlay(const FAssetData& AssetData) const
-{
-	const FSlateBrush* Icon = FSlateIconFinder::FindIconBrushForClass(URenderGridBlueprint::StaticClass());
-
-	return SNew(SBorder)
-		.BorderImage(FAppStyle::GetNoBrush())
-		.Visibility(EVisibility::HitTestInvisible)
-		.Padding(FMargin(0.0f, 0.0f, 0.0f, 3.0f))
-		.HAlign(HAlign_Right)
-		.VAlign(VAlign_Bottom)
-		[
-			SNew(SImage)
-			.Image(Icon)
-		];
-}
-
 UFactory* UE::RenderGrid::Private::FRenderGridBlueprintActions::GetFactoryForBlueprintType(UBlueprint* InBlueprint) const
 {
 	URenderGridBlueprintFactory* RenderGridBlueprintFactory = NewObject<URenderGridBlueprintFactory>();
diff --git a/Engine/Plugins/Experimental/RenderGrid/Source/RenderGridEditor/Private/AssetTypeActions/RenderGridBlueprintActions.h b/Engine/Plugins/Experimental/RenderGrid/Source/RenderGridEditor/Private/AssetTypeActions/RenderGridBlueprintActions.h
--- a/Engine/Plugins/Experimental/RenderGrid/Source/RenderGridEditor/Private/AssetTypeActions/RenderGridBlueprintActions.h
+++ b/Engine/Plugins/Experimental/RenderGrid/Source/RenderGridEditor/Private/AssetTypeActions/RenderGridBlueprintActions.h
@@ -20,7 +20,6 @@ namespace UE::RenderGrid::Private
 		virtual UClass* GetSupportedClass() const override;
 		virtual void OpenAssetEditor(const TArray<UObject*>& InObjects, TSharedPtr<IToolkitHost> EditWithinLevelEditor = TSharedPtr<IToolkitHost>()) override;
 		virtual uint32 GetCategories() override;
-		virtual TSharedPtr<SWidget> GetThumbnailOverlay(const FAssetData& AssetData) const override;
 		//~ End IAssetTypeActions Interface
 
 		//~ Begin FAssetTypeActions_Blueprint interface
diff --git a/Engine/Plugins/Experimental/RenderGrid/Source/RenderGridEditor/Private/Styles/RenderGridEditorStyle.cpp b/Engine/Plugins/Experimental/RenderGrid/Source/RenderGridEditor/Private/Styles/RenderGridEditorStyle.cpp
--- a/Engine/Plugins/Experimental/RenderGrid/Source/RenderGridEditor/Private/Styles/RenderGridEditorStyle.cpp
+++ b/Engine/Plugins/Experimental/RenderGrid/Source/RenderGridEditor/Private/Styles/RenderGridEditorStyle.cpp
@@ -36,7 +36,7 @@ void UE::RenderGrid::Private::FRenderGridEditorStyle::Shutdown()
 
 FName UE::RenderGrid::Private::FRenderGridEditorStyle::GetStyleSetName()
 {
-	static FName StyleSetName(TEXT("LevelSnapshotsEditor"));
+	static FName StyleSetName(TEXT("RenderGridEditorStyle"));
 	return StyleSetName;
 }
 
@@ -73,6 +73,11 @@ TSharedRef<FSlateStyleSet> UE::RenderGrid::Private::FRenderGridEditorStyle::Crea
 		Style->SetContentRoot(FPaths::Combine(Plugin->GetBaseDir(), TEXT("Resources")));
 	}
 
+	Style->Set("ClassIcon.RenderGrid", new IMAGE_BRUSH_SVG("Icons/RenderGrid_16", Icon16x16));
+	Style->Set("ClassIcon.RenderGridBlueprint", new IMAGE_BRUSH_SVG("Icons/RenderGrid_16", Icon16x16));
+	Style->Set("ClassThumbnail.RenderGrid", new IMAGE_BRUSH_SVG("Icons/RenderGrid_64", Icon64x64));
+	Style->Set("ClassThumbnail.RenderGridBlueprint", new IMAGE_BRUSH_SVG("Icons/RenderGrid_64", Icon64x64));
+
 	Style->Set("Invisible",
 		FButtonStyle()
 		.SetNormal(FSlateNoResource())
