Fixed an issue with negative frames not working in the rendered viewer mode.

diff --git a/Engine/Plugins/Experimental/RenderGrid/Source/RenderGrid/Private/RenderGrid/RenderGridManager.cpp b/Engine/Plugins/Experimental/RenderGrid/Source/RenderGrid/Private/RenderGrid/RenderGridManager.cpp
--- a/Engine/Plugins/Experimental/RenderGrid/Source/RenderGrid/Private/RenderGrid/RenderGridManager.cpp
+++ b/Engine/Plugins/Experimental/RenderGrid/Source/RenderGrid/Private/RenderGrid/RenderGridManager.cpp
@@ -88,14 +88,7 @@ URenderGridQueue* UE::RenderGrid::FRenderGridManager::RenderPreviewFrame(const F
 	{
 		if (StartFrameOfRender.IsSet())
 		{
-			if (bSuccess)
-			{
-				StartFrameOfRenders.Add(JobId, StartFrameOfRender.Get(0));
-			}
-			else
-			{
-				StartFrameOfRenders.Remove(JobId);
-			}
+			StartFrameOfRenders.Add(JobId, StartFrameOfRender.Get(0));
 		}
 		Callback.ExecuteIfBound(bSuccess);
 	});
diff --git a/Engine/Plugins/Experimental/RenderGrid/Source/RenderGrid/Private/RenderGrid/RenderGridQueue.cpp b/Engine/Plugins/Experimental/RenderGrid/Source/RenderGrid/Private/RenderGrid/RenderGridQueue.cpp
--- a/Engine/Plugins/Experimental/RenderGrid/Source/RenderGrid/Private/RenderGrid/RenderGridQueue.cpp
+++ b/Engine/Plugins/Experimental/RenderGrid/Source/RenderGrid/Private/RenderGrid/RenderGridQueue.cpp
@@ -478,6 +478,7 @@ void URenderGridQueue::OnProcessJob(URenderGridJob* Job)
 
 void URenderGridQueue::OnFinish()
 {
+	bool bSuccess = !bCanceled;
 	Cancel();// to prevent any new jobs from being added to it
 
 	RenderGrid->EndBatchRender(this);
@@ -486,5 +487,5 @@ void URenderGridQueue::OnFinish()
 	PreviousFrameLimitSettings = FRenderGridPreviousEngineFpsSettings();
 
 	RemoveFromRoot();
-	OnExecuteFinishedDelegate.Broadcast(this, !bCanceled);
+	OnExecuteFinishedDelegate.Broadcast(this, bSuccess);
 }
