In the RenderGrid plugin:
- Fixed a bug with the remote control fields.
- Added a context menu to the job list rows.
- Added a shortcut for the delete job action.
- Changed the duplicate button icon.

diff --git a/Engine/Plugins/Experimental/RenderGrid/Source/RenderGridEditor/Private/Commands/RenderGridEditorCommands.cpp b/Engine/Plugins/Experimental/RenderGrid/Source/RenderGridEditor/Private/Commands/RenderGridEditorCommands.cpp
--- a/Engine/Plugins/Experimental/RenderGrid/Source/RenderGridEditor/Private/Commands/RenderGridEditorCommands.cpp
+++ b/Engine/Plugins/Experimental/RenderGrid/Source/RenderGridEditor/Private/Commands/RenderGridEditorCommands.cpp
@@ -9,7 +9,7 @@ void UE::RenderGrid::Private::FRenderGridEditorCommands::RegisterCommands()
 {
 	UI_COMMAND(AddJob, "Add", "Creates and adds a new job instance to the grid.", EUserInterfaceActionType::Button, FInputChord());
 	UI_COMMAND(DuplicateJob, "Duplicate", "Duplicates the selected job instances and adds them to the grid.", EUserInterfaceActionType::Button, FInputChord());
-	UI_COMMAND(DeleteJob, "Delete", "Deletes the selected job instances from the grid.", EUserInterfaceActionType::Button, FInputChord());
+	UI_COMMAND(DeleteJob, "Delete", "Deletes the selected job instances from the grid.", EUserInterfaceActionType::Button, FInputChord(EModifierKey::None, EKeys::Delete));
 
 	UI_COMMAND(BatchRenderList, "Render", "Renders the enabled job instances of the grid.", EUserInterfaceActionType::Button, FInputChord());
 }
diff --git a/Engine/Plugins/Experimental/RenderGrid/Source/RenderGridEditor/Private/Toolkit/RenderGridEditorToolbar.cpp b/Engine/Plugins/Experimental/RenderGrid/Source/RenderGridEditor/Private/Toolkit/RenderGridEditorToolbar.cpp
--- a/Engine/Plugins/Experimental/RenderGrid/Source/RenderGridEditor/Private/Toolkit/RenderGridEditorToolbar.cpp
+++ b/Engine/Plugins/Experimental/RenderGrid/Source/RenderGridEditor/Private/Toolkit/RenderGridEditorToolbar.cpp
@@ -81,7 +81,7 @@ void UE::RenderGrid::Private::FRenderGridBlueprintEditorToolbar::AddListingModeT
 		FRenderGridEditorCommands::Get().DuplicateJob,
 		TAttribute<FText>(),
 		TAttribute<FText>(),
-		FSlateIcon(FAppStyle::GetAppStyleSetName(), "GenericCommands.Duplicate")
+		FSlateIcon(FAppStyle::GetAppStyleSetName(), "Icons.Duplicate")
 	));
 
 	Section.AddEntry(FToolMenuEntry::InitToolBarButton(
diff --git a/Engine/Plugins/Experimental/RenderGrid/Source/RenderGridEditor/Private/UI/Components/SRenderGridRemoteControlField.cpp b/Engine/Plugins/Experimental/RenderGrid/Source/RenderGridEditor/Private/UI/Components/SRenderGridRemoteControlField.cpp
--- a/Engine/Plugins/Experimental/RenderGrid/Source/RenderGridEditor/Private/UI/Components/SRenderGridRemoteControlField.cpp
+++ b/Engine/Plugins/Experimental/RenderGrid/Source/RenderGridEditor/Private/UI/Components/SRenderGridRemoteControlField.cpp
@@ -50,10 +50,23 @@ TSharedPtr<UE::RenderGrid::Private::SRenderGridRemoteControlTreeNode> UE::Render
 	return SNew(SRenderGridRemoteControlField, StaticCastSharedPtr<FRemoteControlField>(Args.Entity), Args.ColumnSizeData).Preset(Args.Preset);
 }
 
+void UE::RenderGrid::Private::SRenderGridRemoteControlField::Tick(const FGeometry&, const double, const float)
+{
+	if (FramesUntilRerender > 0)
+	{
+		FramesUntilRerender--;
+		if (FramesUntilRerender <= 0)
+		{
+			Refresh();
+		}
+	}
+}
+
 void UE::RenderGrid::Private::SRenderGridRemoteControlField::Construct(const FArguments& InArgs, TWeakPtr<FRemoteControlField> InField, FRenderGridRemoteControlColumnSizeData InColumnSizeData)
 {
 	FieldWeakPtr = MoveTemp(InField);
 	ColumnSizeData = MoveTemp(InColumnSizeData);
+	FramesUntilRerender = 0;
 
 	if (const TSharedPtr<FRemoteControlField> Field = FieldWeakPtr.Pin())
 	{
@@ -65,6 +78,16 @@ void UE::RenderGrid::Private::SRenderGridRemoteControlField::Construct(const FAr
 		if (Field->FieldType == EExposedFieldType::Property)
 		{
 			ConstructPropertyWidget();
+
+			if (Generator.IsValid())
+			{
+				Generator.Get()->OnRowsRefreshed().RemoveAll(this);
+				Generator.Get()->OnRowsRefreshed().AddLambda([this]()
+				{
+					Generator.Get()->OnRowsRefreshed().RemoveAll(this);
+					FramesUntilRerender = 2;
+				});
+			}
 		}
 	}
 }
diff --git a/Engine/Plugins/Experimental/RenderGrid/Source/RenderGridEditor/Private/UI/Components/SRenderGridRemoteControlField.h b/Engine/Plugins/Experimental/RenderGrid/Source/RenderGridEditor/Private/UI/Components/SRenderGridRemoteControlField.h
--- a/Engine/Plugins/Experimental/RenderGrid/Source/RenderGridEditor/Private/UI/Components/SRenderGridRemoteControlField.h
+++ b/Engine/Plugins/Experimental/RenderGrid/Source/RenderGridEditor/Private/UI/Components/SRenderGridRemoteControlField.h
@@ -52,6 +52,7 @@ namespace UE::RenderGrid::Private
 
 		static TSharedPtr<SRenderGridRemoteControlTreeNode> MakeInstance(const FRenderGridRemoteControlGenerateWidgetArgs& Args);
 
+		virtual void Tick(const FGeometry& AllottedGeometry, const double InCurrentTime, const float InDeltaTime) override;
 		void Construct(const FArguments& InArgs, TWeakPtr<FRemoteControlField> Field, FRenderGridRemoteControlColumnSizeData ColumnSizeData);
 
 		//~ SRenderGridTreeNode Interface 
@@ -76,21 +77,24 @@ namespace UE::RenderGrid::Private
 	private:
 		/** Construct a property widget. */
 		TSharedRef<SWidget> ConstructWidget();
-		
+
 		/** Create the wrapper around the field value widget. */
 		TSharedRef<SWidget> MakeFieldWidget(const TSharedRef<SWidget>& InWidget);
-		
+
 		/** Construct this field widget as a property widget. */
 		void ConstructPropertyWidget();
-		
+
 	private:
 		/** Weak pointer to the underlying RC Field. */
 		TWeakPtr<FRemoteControlField> FieldWeakPtr;
-		
+
 		/** This exposed field's child widgets (ie. An array's rows) */
 		TArray<TSharedPtr<SRenderGridRemoteControlFieldChildNode>> ChildWidgets;
-		
+
 		/** The property row generator. */
 		TSharedPtr<IPropertyRowGenerator> Generator;
+
+		/** Has a value of 1 if it should call Render() next frame, 2+ if it should subtract 1 from its value next frame, 0 and below and it won't do anything. */
+		int32 FramesUntilRerender;
 	};
 }
diff --git a/Engine/Plugins/Experimental/RenderGrid/Source/RenderGridEditor/Private/UI/SRenderGridJobList.cpp b/Engine/Plugins/Experimental/RenderGrid/Source/RenderGridEditor/Private/UI/SRenderGridJobList.cpp
--- a/Engine/Plugins/Experimental/RenderGrid/Source/RenderGridEditor/Private/UI/SRenderGridJobList.cpp
+++ b/Engine/Plugins/Experimental/RenderGrid/Source/RenderGridEditor/Private/UI/SRenderGridJobList.cpp
@@ -9,6 +9,7 @@
 #include "RenderGrid/RenderGridQueue.h"
 #include "IRenderGridEditor.h"
 
+#include "Framework/MultiBox/MultiBoxBuilder.h"
 #include "Misc/MessageDialog.h"
 #include "MoviePipelineExecutor.h"
 #include "MoviePipelineOutputSetting.h"
@@ -338,6 +339,42 @@ void UE::RenderGrid::Private::SRenderGridJobListTableRow::Construct(const FArgum
 		InOwnerTableView);
 }
 
+FReply UE::RenderGrid::Private::SRenderGridJobListTableRow::OnMouseButtonUp(const FGeometry& MyGeometry, const FPointerEvent& MouseEvent)
+{
+	if (MouseEvent.GetEffectingButton() != EKeys::RightMouseButton)
+	{
+		return SMultiColumnTableRow<URenderGridJob*>::OnMouseButtonUp(MyGeometry, MouseEvent);
+	}
+
+	FMenuBuilder MenuBuilder(true, nullptr);
+	{
+		MenuBuilder.BeginSection("RenderGridJobListRowContextMenuSection", LOCTEXT("RenderGridJobListRowContextMenuHeading", "Options"));
+		{
+			MenuBuilder.AddMenuEntry(
+				LOCTEXT("DuplicateRenderGridJobListRowLabel", "Duplicate"),
+				LOCTEXT("DuplicateRenderGridJobListTooltip", "Creates a copy of the job and adds it to the grid."),
+				FSlateIcon(FAppStyle::GetAppStyleSetName(), "Icons.Duplicate"),
+				FUIAction(FExecuteAction::CreateSP(this, &SRenderGridJobListTableRow::DuplicateJob))
+			);
+			MenuBuilder.AddMenuEntry(
+				LOCTEXT("DeleteRenderGridJobListRowLabel", "Delete"),
+				LOCTEXT("DeleteRenderGridJobListTooltip", "Removes the job from the grid."),
+				FSlateIcon(FAppStyle::GetAppStyleSetName(), "Icons.Delete"),
+				FUIAction(FExecuteAction::CreateSP(this, &SRenderGridJobListTableRow::DeleteJob))
+			);
+		}
+		MenuBuilder.EndSection();
+	}
+	TSharedPtr<SWidget> MenuContent = MenuBuilder.MakeWidget();
+
+	if (MenuContent.IsValid() && (MouseEvent.GetEventPath() != nullptr))
+	{
+		FWidgetPath WidgetPath = *MouseEvent.GetEventPath();
+		FSlateApplication::Get().PushMenu(WidgetPath.Widgets.Last().Widget, WidgetPath, MenuContent.ToSharedRef(), MouseEvent.GetScreenSpacePosition(), FPopupTransitionEffect(FPopupTransitionEffect::ContextMenu));
+	}
+	return FReply::Handled();
+}
+
 TOptional<EItemDropZone> UE::RenderGrid::Private::SRenderGridJobListTableRow::OnCanAcceptDrop(const FDragDropEvent& InEvent, EItemDropZone InItemDropZone, URenderGridJob* InJob)
 {
 	if (BlueprintEditorWeakPtr.IsValid() && InEvent.GetOperationAs<FRenderGridJobListTableRowDragDropOp>())
@@ -596,6 +633,32 @@ TSharedRef<SWidget> UE::RenderGrid::Private::SRenderGridJobListTableRow::Generat
 	return SNullWidget::NullWidget;
 }
 
+void UE::RenderGrid::Private::SRenderGridJobListTableRow::DuplicateJob()
+{
+	if (const TSharedPtr<IRenderGridEditor> BlueprintEditor = BlueprintEditorWeakPtr.Pin())
+	{
+		if (URenderGrid* RenderGrid = BlueprintEditor->GetInstance(); IsValid(RenderGrid))
+		{
+			RenderGrid->DuplicateAndAddRenderGridJob(RenderGridJob);
+			BlueprintEditor->MarkAsModified();
+			BlueprintEditor->OnRenderGridChanged().Broadcast();
+		}
+	}
+}
+
+void UE::RenderGrid::Private::SRenderGridJobListTableRow::DeleteJob()
+{
+	if (const TSharedPtr<IRenderGridEditor> BlueprintEditor = BlueprintEditorWeakPtr.Pin())
+	{
+		if (URenderGrid* RenderGrid = BlueprintEditor->GetInstance(); IsValid(RenderGrid))
+		{
+			RenderGrid->RemoveRenderGridJob(RenderGridJob);
+			BlueprintEditor->MarkAsModified();
+			BlueprintEditor->OnRenderGridChanged().Broadcast();
+		}
+	}
+}
+
 FText UE::RenderGrid::Private::SRenderGridJobListTableRow::GetRenderStatusText() const
 {
 	if (const TSharedPtr<IRenderGridEditor> BlueprintEditor = BlueprintEditorWeakPtr.Pin())
diff --git a/Engine/Plugins/Experimental/RenderGrid/Source/RenderGridEditor/Private/UI/SRenderGridJobList.h b/Engine/Plugins/Experimental/RenderGrid/Source/RenderGridEditor/Private/UI/SRenderGridJobList.h
--- a/Engine/Plugins/Experimental/RenderGrid/Source/RenderGridEditor/Private/UI/SRenderGridJobList.h
+++ b/Engine/Plugins/Experimental/RenderGrid/Source/RenderGridEditor/Private/UI/SRenderGridJobList.h
@@ -106,9 +106,18 @@ namespace UE::RenderGrid::Private
 		SLATE_END_ARGS()
 
 		void Construct(const FArguments& InArgs, const TSharedRef<STableViewBase>& InOwnerTableView, TWeakPtr<IRenderGridEditor> InBlueprintEditor, URenderGridJob* InRenderGridJob, const TSharedPtr<SRenderGridJobList>& InJobListWidget);
+
+		//~ Begin SWidget Interface
+		virtual FReply OnMouseButtonUp(const FGeometry& MyGeometry, const FPointerEvent& MouseEvent) override;
+		//~ End SWidget Interface
+
 		TOptional<EItemDropZone> OnCanAcceptDrop(const FDragDropEvent& InEvent, EItemDropZone InItemDropZone, URenderGridJob* InJob);
 		FReply OnAcceptDrop(const FDragDropEvent& InEvent, EItemDropZone InItemDropZone, URenderGridJob* InJob);
 		virtual TSharedRef<SWidget> GenerateWidgetForColumn(const FName& ColumnName) override;
+
+		void DuplicateJob();
+		void DeleteJob();
+
 		FText GetRenderStatusText() const;
 
 	private:
diff --git a/Engine/Plugins/Experimental/RenderGrid/Source/RenderGridEditor/Private/UI/SRenderGridViewerPreview.h b/Engine/Plugins/Experimental/RenderGrid/Source/RenderGridEditor/Private/UI/SRenderGridViewerPreview.h
--- a/Engine/Plugins/Experimental/RenderGrid/Source/RenderGridEditor/Private/UI/SRenderGridViewerPreview.h
+++ b/Engine/Plugins/Experimental/RenderGrid/Source/RenderGridEditor/Private/UI/SRenderGridViewerPreview.h
@@ -79,7 +79,7 @@ namespace UE::RenderGrid::Private
 		/** A reference to the job that's currently rendering. */
 		TWeakObjectPtr<URenderGridJob> SelectedJobWeakPtr;
 
-		/** 1 if it should call RenderNewPreview() next frame, 2+ if it should subtract 1 from its value next frame, 0 and below and it won't do anything. */
+		/** Has a value of 1 if it should call RenderNewPreview() next frame, 2+ if it should subtract 1 from its value next frame, 0 and below and it won't do anything. */
 		int32 FramesUntilRenderNewPreview;
 
 	private:
