- Added a hotkey for add: CTRL+N.
- Added a hotkey for duplicate: CTRL+D.
- The context menu will now only affect selected jobs (rather than the right-clicked job).
- Added the "Add Job" function to the context menu.
- Right-clicking the background of the job list will now also show up the context menu.
- Made the "Duplicate Job" and "Delete Job" action in the context menu automatically be hidden when no job is selected.
- Changed the remote control properties header to "Source".
- Added additional padding above the frame slider.
- Fixed non-unity compiler issues.
- Added blueprints to obtain the rendering process of the queue, as well as of each individual job.
- Fixed a bug, which caused the End Batch Render event to fire, even when no batch rendering was started.

diff --git a/Engine/Plugins/Experimental/RenderGrid/Source/RenderGrid/Private/RenderGrid/RenderGrid.cpp b/Engine/Plugins/Experimental/RenderGrid/Source/RenderGrid/Private/RenderGrid/RenderGrid.cpp
--- a/Engine/Plugins/Experimental/RenderGrid/Source/RenderGrid/Private/RenderGrid/RenderGrid.cpp
+++ b/Engine/Plugins/Experimental/RenderGrid/Source/RenderGrid/Private/RenderGrid/RenderGrid.cpp
@@ -548,6 +548,13 @@ void URenderGrid::EndEditor()
 	bExecutingBlueprintEvent = false;
 }
 
+void URenderGrid::Tick(float DeltaTime)
+{
+	bExecutingBlueprintEvent = true;
+	ReceiveTick(DeltaTime);
+	bExecutingBlueprintEvent = false;
+}
+
 void URenderGrid::BeginBatchRender(URenderGridQueue* Queue)
 {
 	bExecutingBlueprintEvent = true;
diff --git a/Engine/Plugins/Experimental/RenderGrid/Source/RenderGrid/Private/RenderGrid/RenderGridQueue.cpp b/Engine/Plugins/Experimental/RenderGrid/Source/RenderGrid/Private/RenderGrid/RenderGridQueue.cpp
--- a/Engine/Plugins/Experimental/RenderGrid/Source/RenderGrid/Private/RenderGrid/RenderGridQueue.cpp
+++ b/Engine/Plugins/Experimental/RenderGrid/Source/RenderGrid/Private/RenderGrid/RenderGridQueue.cpp
@@ -386,6 +386,99 @@ FString URenderGridQueue::GetJobStatus(URenderGridJob* Job) const
 	return TEXT("");
 }
 
+TArray<URenderGridJob*> URenderGridQueue::GetJobs() const
+{
+	TArray<URenderGridJob*> Jobs;
+	Jobs.Reserve(Entries.Num());
+	for (const TTuple<TObjectPtr<URenderGridJob>, TObjectPtr<URenderGridMoviePipelineRenderJob>> EntryEntry : Entries)
+	{
+		if (const TObjectPtr<URenderGridJob> Job = EntryEntry.Key; IsValid(Job))
+		{
+			if (const TObjectPtr<URenderGridMoviePipelineRenderJob> Entry = EntryEntry.Value; IsValid(Entry))
+			{
+				if (Entry->CanExecute())
+				{
+					Jobs.Add(Job);
+				}
+			}
+		}
+	}
+	return Jobs;
+}
+
+int32 URenderGridQueue::GetJobsCount() const
+{
+	int32 JobsCount = 0;
+	for (const TTuple<TObjectPtr<URenderGridJob>, TObjectPtr<URenderGridMoviePipelineRenderJob>> EntryEntry : Entries)
+	{
+		if (const TObjectPtr<URenderGridJob> Job = EntryEntry.Key; IsValid(Job))
+		{
+			if (const TObjectPtr<URenderGridMoviePipelineRenderJob> Entry = EntryEntry.Value; IsValid(Entry))
+			{
+				if (Entry->CanExecute())
+				{
+					JobsCount++;
+				}
+			}
+		}
+	}
+	return JobsCount;
+}
+
+int32 URenderGridQueue::GetJobsRemainingCount() const
+{
+	if (bCanceled)
+	{
+		return 0;
+	}
+
+	int32 JobsCount = 0;
+	for (const TObjectPtr<URenderGridJob> Job : RemainingJobs)
+	{
+		if (const TObjectPtr<URenderGridMoviePipelineRenderJob>* EntryPtr = Entries.Find(Job))
+		{
+			if (const TObjectPtr<URenderGridMoviePipelineRenderJob> Entry = *EntryPtr; IsValid(Entry))
+			{
+				if (Entry->CanExecute())
+				{
+					JobsCount++;
+				}
+			}
+		}
+	}
+	return JobsCount + 1; // one is removed when it starts rendering, but this job is still remaining of course, so let's add 1 here
+}
+
+int32 URenderGridQueue::GetJobsCompletedCount() const
+{
+	return (GetJobsCount() - GetJobsRemainingCount());
+}
+
+float URenderGridQueue::GetStatusPercentage() const
+{
+	if (bCanceled)
+	{
+		return 100.0;
+	}
+
+	double JobsCount = GetJobsCount();
+	if (JobsCount <= 0)
+	{
+		return 0.0;
+	}
+	double RemainingCount = GetJobsRemainingCount() - 0.5; // rendering a job, but we can't get the progression from the job, so let's say we're halfway there (0.0 started, 0.5 halfway, 1.0 finished)
+	return FMath::Clamp(100.0 - ((RemainingCount / JobsCount) * 100), 0.0, 100.0);
+}
+
+FString URenderGridQueue::GetStatus() const
+{
+	if (bCanceled)
+	{
+		return TEXT("Done");
+	}
+	return TEXT("Rendering...");
+}
+
 
 void URenderGridQueue::OnStart()
 {
@@ -481,7 +574,10 @@ void URenderGridQueue::OnFinish()
 	bool bSuccess = !bCanceled;
 	Cancel();// to prevent any new jobs from being added to it
 
-	RenderGrid->EndBatchRender(this);
+	if (Args.bIsBatchRender)
+	{
+		RenderGrid->EndBatchRender(this);
+	}
 
 	UE::RenderGrid::Private::FRenderGridUtils::RestoreFpsLimit(PreviousFrameLimitSettings);
 	PreviousFrameLimitSettings = FRenderGridPreviousEngineFpsSettings();
diff --git a/Engine/Plugins/Experimental/RenderGrid/Source/RenderGrid/Public/RenderGrid/RenderGrid.h b/Engine/Plugins/Experimental/RenderGrid/Source/RenderGrid/Public/RenderGrid/RenderGrid.h
--- a/Engine/Plugins/Experimental/RenderGrid/Source/RenderGrid/Public/RenderGrid/RenderGrid.h
+++ b/Engine/Plugins/Experimental/RenderGrid/Source/RenderGrid/Public/RenderGrid/RenderGrid.h
@@ -303,6 +303,7 @@ public:
 		return {
 			TEXT("ReceiveBeginEditor"),
 			TEXT("ReceiveEndEditor"),
+			TEXT("ReceiveTick"),
 			TEXT("ReceiveBeginBatchRender"),
 			TEXT("ReceiveEndBatchRender"),
 			TEXT("ReceiveBeginJobRender"),
@@ -329,6 +330,12 @@ protected:
 	UFUNCTION(BlueprintImplementableEvent, Meta=(DisplayName="EndEditor"))
 	void ReceiveEndEditor();
 
+	/**
+	 * The tick event, will only execute when the asset is open in the editor.
+	 */
+	UFUNCTION(BlueprintImplementableEvent, Meta=(DisplayName="Tick"))
+	void ReceiveTick(float DeltaTime);
+
 	/**
 	 * Event for when batch rendering begins.
 	 * 
@@ -390,6 +397,9 @@ public:
 	/** Overridable native event for when this asset is closed in the editor. The asset will also be closed and reopened during a blueprint compile. */
 	virtual void EndEditor();
 
+	/** Overridable native event for when this asset is ticked by the editor. */
+	virtual void Tick(float DeltaTime);
+
 	/** Overridable native event for when batch rendering begins. */
 	virtual void BeginBatchRender(URenderGridQueue* Queue);
 
@@ -539,11 +549,11 @@ private:
 
 
 	/** The type of the properties that a job in this grid can have. */
-	UPROPERTY(/*EditInstanceOnly, Category="Render Grid|Grid", Meta=(DisplayName="Properties Type", AllowPrivateAccess="true")*/)
+	UPROPERTY(/*EditInstanceOnly, Category="Render Grid|Source", Meta=(DisplayName="Properties Type", AllowPrivateAccess="true")*/)
 	ERenderGridPropsSourceType PropsSourceType;
 
 	/** The remote control properties that a job in this grid can have, only use this if PropsSourceType is ERenderGridPropsSourceType::RemoteControl. */
-	UPROPERTY(EditInstanceOnly, Category="Render Grid|Grid", Meta=(DisplayName="Remote Control Preset", AllowPrivateAccess="true", EditCondition="PropsSourceType == ERenderGridPropsSourceType::RemoteControl", EditConditionHides))
+	UPROPERTY(EditInstanceOnly, Category="Render Grid|Source", Meta=(DisplayName="Remote Control Preset", AllowPrivateAccess="true", EditCondition="PropsSourceType == ERenderGridPropsSourceType::RemoteControl", EditConditionHides))
 	TObjectPtr<URemoteControlPreset> PropsSourceOrigin_RemoteControl;
 
 
diff --git a/Engine/Plugins/Experimental/RenderGrid/Source/RenderGrid/Public/RenderGrid/RenderGridQueue.h b/Engine/Plugins/Experimental/RenderGrid/Source/RenderGrid/Public/RenderGrid/RenderGridQueue.h
--- a/Engine/Plugins/Experimental/RenderGrid/Source/RenderGrid/Public/RenderGrid/RenderGridQueue.h
+++ b/Engine/Plugins/Experimental/RenderGrid/Source/RenderGrid/Public/RenderGrid/RenderGridQueue.h
@@ -102,6 +102,9 @@ public:
 	/** Returns true if this render job was canceled (which for example can be caused by calling Cancel(), or by closing the render popup). */
 	bool IsCanceled() const { return bCanceled; }
 
+	/** Returns true if this render job can render (otherwise it will be skipped). */
+	bool CanExecute() const { return bCanExecute; }
+
 private:
 	void ComputePlaybackContext(bool& bOutAllowBinding);
 	void ExecuteFinished(UMoviePipelineExecutorBase* InPipelineExecutor, const bool bSuccess);
@@ -184,9 +187,33 @@ public:
 	bool IsCanceled() const { return bCanceled; }
 
 	/** Retrieves the rendering status of the given render grid job. */
-	UFUNCTION(BlueprintCallable, Category="Render Grid|Queue", Meta=(Keywords="render progress"))
+	UFUNCTION(BlueprintCallable, Category="Render Grid|Queue", Meta=(Keywords="rendering progression obtain text"))
 	FString GetJobStatus(URenderGridJob* Job) const;
 
+	/** Returns all the jobs that have been and will be rendered. */
+	UFUNCTION(BlueprintCallable, Category="Render Grid|Queue", Meta=(Keywords="rendering progression obtain"))
+	TArray<URenderGridJob*> GetJobs() const;
+
+	/** Returns the number of jobs that have been and will be rendered. */
+	UFUNCTION(BlueprintCallable, Category="Render Grid|Queue", Meta=(Keywords="rendering progression obtain number amount"))
+	int32 GetJobsCount() const;
+
+	/** Returns the number of jobs that have finished rendering. Basically just returns [Get Jobs Count] minus [Get Jobs Remaining Count]. */
+	UFUNCTION(BlueprintCallable, Category="Render Grid|Queue", Meta=(Keywords="rendering progression obtain number amount finished"))
+	int32 GetJobsCompletedCount() const;
+
+	/** Returns the percentage of jobs finished, this includes the progression of the job that is currently rendering. */
+	UFUNCTION(BlueprintCallable, Category="Render Grid|Queue", Meta=(Keywords="rendering progression obtain number amount finished"))
+	float GetStatusPercentage() const;
+
+	/** Returns the number of jobs that are still left to render, includes the job that is currently rendering. */
+	UFUNCTION(BlueprintCallable, Category="Render Grid|Queue", Meta=(Keywords="rendering progression obtain number amount finished"))
+	int32 GetJobsRemainingCount() const;
+
+	/** Returns the status of the rendering process. */
+	UFUNCTION(BlueprintCallable, Category="Render Grid|Queue", Meta=(Keywords="rendering progression obtain text"))
+	FString GetStatus() const;
+
 protected:
 	void OnStart();
 	void OnProcessJob(URenderGridJob* Job);
@@ -205,7 +232,7 @@ protected:
 
 	/** The render grid jobs that are to be rendered, mapped to the movie pipeline render job of each specific render grid job. */
 	UPROPERTY(Transient)
-	TMap<TObjectPtr<const URenderGridJob>, TObjectPtr<URenderGridMoviePipelineRenderJob>> Entries;
+	TMap<TObjectPtr<URenderGridJob>, TObjectPtr<URenderGridMoviePipelineRenderJob>> Entries;
 
 	/** The render grid of the given render grid job that will be rendered. */
 	UPROPERTY(Transient)
diff --git a/Engine/Plugins/Experimental/RenderGrid/Source/RenderGridEditor/Private/Commands/RenderGridEditorCommands.cpp b/Engine/Plugins/Experimental/RenderGrid/Source/RenderGridEditor/Private/Commands/RenderGridEditorCommands.cpp
--- a/Engine/Plugins/Experimental/RenderGrid/Source/RenderGridEditor/Private/Commands/RenderGridEditorCommands.cpp
+++ b/Engine/Plugins/Experimental/RenderGrid/Source/RenderGridEditor/Private/Commands/RenderGridEditorCommands.cpp
@@ -7,8 +7,8 @@
 
 void UE::RenderGrid::Private::FRenderGridEditorCommands::RegisterCommands()
 {
-	UI_COMMAND(AddJob, "Add", "Creates and adds a new job instance to the grid.", EUserInterfaceActionType::Button, FInputChord());
-	UI_COMMAND(DuplicateJob, "Duplicate", "Duplicates the selected job instances and adds them to the grid.", EUserInterfaceActionType::Button, FInputChord());
+	UI_COMMAND(AddJob, "Add", "Creates and adds a new job instance to the grid.", EUserInterfaceActionType::Button, FInputChord(EModifierKey::Control, EKeys::N));
+	UI_COMMAND(DuplicateJob, "Duplicate", "Duplicates the selected job instances and adds them to the grid.", EUserInterfaceActionType::Button, FInputChord(EModifierKey::Control, EKeys::D));
 	UI_COMMAND(DeleteJob, "Delete", "Deletes the selected job instances from the grid.", EUserInterfaceActionType::Button, FInputChord(EModifierKey::None, EKeys::Delete));
 
 	UI_COMMAND(BatchRenderList, "Render", "Renders the enabled job instances of the grid.", EUserInterfaceActionType::Button, FInputChord());
diff --git a/Engine/Plugins/Experimental/RenderGrid/Source/RenderGridEditor/Private/Toolkit/RenderGridEditor.cpp b/Engine/Plugins/Experimental/RenderGrid/Source/RenderGridEditor/Private/Toolkit/RenderGridEditor.cpp
--- a/Engine/Plugins/Experimental/RenderGrid/Source/RenderGridEditor/Private/Toolkit/RenderGridEditor.cpp
+++ b/Engine/Plugins/Experimental/RenderGrid/Source/RenderGridEditor/Private/Toolkit/RenderGridEditor.cpp
@@ -96,13 +96,6 @@ UE::RenderGrid::Private::FRenderGridEditor::FRenderGridEditor()
 
 UE::RenderGrid::Private::FRenderGridEditor::~FRenderGridEditor()
 {
-	URenderGridBlueprint* RenderGridBlueprint = FRenderGridEditor::GetRenderGridBlueprint();
-	if (IsValid(RenderGridBlueprint))
-	{
-		RenderGridEditorClosedDelegate.Broadcast(this, RenderGridBlueprint);
-	}
-	RenderGridBlueprint = nullptr;
-
 	DestroyInstance();
 }
 
@@ -353,6 +346,11 @@ void UE::RenderGrid::Private::FRenderGridEditor::Tick(float DeltaTime)
 			GetBlueprintObj()->SetObjectBeingDebugged(nullptr);
 		}
 	}
+
+	if (URenderGrid* Instance = GetInstance(); IsValid(Instance))
+	{
+		Instance->Tick(DeltaTime);
+	}
 }
 
 TStatId UE::RenderGrid::Private::FRenderGridEditor::GetStatId() const
@@ -406,7 +404,7 @@ void UE::RenderGrid::Private::FRenderGridEditor::OnFinishedChangingProperties(co
 
 void UE::RenderGrid::Private::FRenderGridEditor::BindCommands()
 {
-	const auto& Commands = FRenderGridEditorCommands::Get();
+	const FRenderGridEditorCommands& Commands = FRenderGridEditorCommands::Get();
 
 	ToolkitCommands->MapAction(Commands.AddJob, FExecuteAction::CreateSP(this, &FRenderGridEditor::AddJobAction));
 	ToolkitCommands->MapAction(Commands.DuplicateJob, FExecuteAction::CreateSP(this, &FRenderGridEditor::DuplicateJobAction));
@@ -501,10 +499,11 @@ void UE::RenderGrid::Private::FRenderGridEditor::BatchRenderListAction()
 		{
 			RenderQueue->OnExecuteFinished().AddRaw(this, &FRenderGridEditor::OnBatchRenderListActionFinished);
 			BatchRenderQueue = RenderQueue;
-			OnRenderGridBatchRenderingStarted().Broadcast(RenderQueue);
 
 			SetIsDebugging(true);
 			RenderQueue->Execute();
+
+			OnRenderGridBatchRenderingStarted().Broadcast(RenderQueue);
 		}
 	}
 }
diff --git a/Engine/Plugins/Experimental/RenderGrid/Source/RenderGridEditor/Private/Toolkit/RenderGridEditor.h b/Engine/Plugins/Experimental/RenderGrid/Source/RenderGridEditor/Private/Toolkit/RenderGridEditor.h
--- a/Engine/Plugins/Experimental/RenderGrid/Source/RenderGridEditor/Private/Toolkit/RenderGridEditor.h
+++ b/Engine/Plugins/Experimental/RenderGrid/Source/RenderGridEditor/Private/Toolkit/RenderGridEditor.h
@@ -86,9 +86,6 @@ namespace UE::RenderGrid::Private
 		/** Immediately rebuilds the render grid instance that is being shown in the editor. */
 		void RefreshInstance();
 
-		/** The delegate that will fire when this editor closes. */
-		FRenderGridEditorClosed& OnRenderGridEditorClosed() { return RenderGridEditorClosedDelegate; }
-
 	private:
 		/** Called whenever the blueprint is structurally changed. */
 		virtual void OnBlueprintChangedImpl(UBlueprint* InBlueprint, bool bIsJustBeingCompiled = false) override;
@@ -149,9 +146,6 @@ namespace UE::RenderGrid::Private
 		FActionMenuContent HandleCreateGraphActionMenu(UEdGraph* InGraph, const FVector2D& InNodePosition, const TArray<UEdGraphPin*>& InDraggedPins, bool bAutoExpand, SGraphEditor::FActionMenuClosed InOnMenuClosed);
 
 	private:
-		/** The delegate that will be fired when this editor closes. */
-		FRenderGridEditorClosed RenderGridEditorClosedDelegate;
-
 		/** The toolbar builder that is used to customize the toolbar of this editor. */
 		TSharedPtr<FRenderGridBlueprintEditorToolbar> RenderGridToolbar;
 
diff --git a/Engine/Plugins/Experimental/RenderGrid/Source/RenderGridEditor/Private/UI/Components/SRenderGridViewerFrameSlider.cpp b/Engine/Plugins/Experimental/RenderGrid/Source/RenderGridEditor/Private/UI/Components/SRenderGridViewerFrameSlider.cpp
--- a/Engine/Plugins/Experimental/RenderGrid/Source/RenderGridEditor/Private/UI/Components/SRenderGridViewerFrameSlider.cpp
+++ b/Engine/Plugins/Experimental/RenderGrid/Source/RenderGridEditor/Private/UI/Components/SRenderGridViewerFrameSlider.cpp
@@ -47,7 +47,7 @@ void UE::RenderGrid::Private::SRenderGridViewerFrameSlider::Construct(const FArg
 			SNew(SBox)
 			.HAlign(HAlign_Fill)
 			.VAlign(VAlign_Fill)
-			.HeightOverride(20.0f)
+			.HeightOverride(39.0f)
 			[
 				FrameSlider.ToSharedRef()
 			]
@@ -56,7 +56,7 @@ void UE::RenderGrid::Private::SRenderGridViewerFrameSlider::Construct(const FArg
 		// slider text
 		+ SVerticalBox::Slot()
 		.AutoHeight()
-		.Padding(5.0f, -3.0f, 5.0f, 3.0f)
+		.Padding(5.0f, -14.0f, 5.0f, 3.0f)
 		[
 			SNew(SHorizontalBox)
 
diff --git a/Engine/Plugins/Experimental/RenderGrid/Source/RenderGridEditor/Private/UI/SRenderGridJobList.cpp b/Engine/Plugins/Experimental/RenderGrid/Source/RenderGridEditor/Private/UI/SRenderGridJobList.cpp
--- a/Engine/Plugins/Experimental/RenderGrid/Source/RenderGridEditor/Private/UI/SRenderGridJobList.cpp
+++ b/Engine/Plugins/Experimental/RenderGrid/Source/RenderGridEditor/Private/UI/SRenderGridJobList.cpp
@@ -2,6 +2,7 @@
 
 #include "UI/SRenderGridJobList.h"
 
+#include "Commands/RenderGridEditorCommands.h"
 #include "UI/Components/SRenderGridDragHandle.h"
 #include "UI/Components/SRenderGridEditableTextBlock.h"
 #include "UI/Components/SRenderGridFileSelectorTextBlock.h"
@@ -58,6 +59,9 @@ BEGIN_SLATE_FUNCTION_BUILD_OPTIMIZATION
 
 void UE::RenderGrid::Private::SRenderGridJobList::Construct(const FArguments& InArgs, TSharedPtr<IRenderGridEditor> InBlueprintEditor)
 {
+	TSharedPtr<SWidget> BaseThis = AsShared();
+	TSharedPtr<SRenderGridJobList> This = StaticCastSharedPtr<SRenderGridJobList>(BaseThis);
+
 	BlueprintEditorWeakPtr = InBlueprintEditor;
 
 	Refresh();
@@ -115,12 +119,7 @@ void UE::RenderGrid::Private::SRenderGridJobList::Construct(const FArguments& In
 			.BorderImage(FAppStyle::GetBrush("ToolPanel.GroupBorder"))
 			.Padding(0.f)
 			[
-				SAssignNew(RenderGridJobListWidget, SListView<URenderGridJob*>)
-				.ItemHeight(20.0f)
-				.OnGenerateRow(this, &SRenderGridJobList::HandleJobListGenerateRow)
-				.OnSelectionChanged(this, &SRenderGridJobList::HandleJobListSelectionChanged)
-				.SelectionMode(ESelectionMode::Multi)
-				.ClearSelectionOnClick(false)
+				SAssignNew(RenderGridJobListWidget, SRenderGridJobListTable, BlueprintEditorWeakPtr, This)
 				.ListItemsSource(&RenderGridJobs)
 				.HeaderRow(
 					SNew(SHeaderRow)
@@ -179,6 +178,58 @@ void UE::RenderGrid::Private::SRenderGridJobList::Construct(const FArguments& In
 
 END_SLATE_FUNCTION_BUILD_OPTIMIZATION
 
+FReply UE::RenderGrid::Private::SRenderGridJobList::OnMouseButtonUp(const FGeometry& MyGeometry, const FPointerEvent& MouseEvent)
+{
+	if (MouseEvent.GetEffectingButton() == EKeys::RightMouseButton)
+	{
+		if (const TSharedPtr<IRenderGridEditor> BlueprintEditor = BlueprintEditorWeakPtr.Pin())
+		{
+			FMenuBuilder MenuBuilder(true, BlueprintEditor->GetToolkitCommands());
+			{
+				MenuBuilder.BeginSection("RenderGridJobListRowContextMenuSection", LOCTEXT("RenderGridJobListRowContextMenuHeading", "Options"));
+				{
+					MenuBuilder.AddMenuEntry(
+						FRenderGridEditorCommands::Get().AddJob,
+						NAME_None,
+						TAttribute<FText>(),
+						TAttribute<FText>(),
+						FSlateIcon(FAppStyle::GetAppStyleSetName(), "Icons.Plus")
+					);
+
+					if (!BlueprintEditor->GetSelectedRenderGridJobs().IsEmpty())
+					{
+						MenuBuilder.AddMenuEntry(
+							FRenderGridEditorCommands::Get().DuplicateJob,
+							NAME_None,
+							TAttribute<FText>(),
+							TAttribute<FText>(),
+							FSlateIcon(FAppStyle::GetAppStyleSetName(), "Icons.Duplicate")
+						);
+
+						MenuBuilder.AddMenuEntry(
+							FRenderGridEditorCommands::Get().DeleteJob,
+							NAME_None,
+							TAttribute<FText>(),
+							TAttribute<FText>(),
+							FSlateIcon(FAppStyle::GetAppStyleSetName(), "Icons.Delete")
+						);
+					}
+				}
+				MenuBuilder.EndSection();
+			}
+			TSharedPtr<SWidget> MenuContent = MenuBuilder.MakeWidget();
+
+			if (MenuContent.IsValid() && (MouseEvent.GetEventPath() != nullptr))
+			{
+				FWidgetPath WidgetPath = *MouseEvent.GetEventPath();
+				FSlateApplication::Get().PushMenu(WidgetPath.Widgets.Last().Widget, WidgetPath, MenuContent.ToSharedRef(), MouseEvent.GetScreenSpacePosition(), FPopupTransitionEffect(FPopupTransitionEffect::ContextMenu));
+			}
+			return FReply::Handled();
+		}
+	}
+	return SCompoundWidget::OnMouseButtonUp(MyGeometry, MouseEvent);
+}
+
 void UE::RenderGrid::Private::SRenderGridJobList::OnRenderGridJobCreated(URenderGridJob* Job)
 {
 	if (!RenderGridJobEnabledHeaderCheckbox.IsValid())
@@ -244,7 +295,6 @@ void UE::RenderGrid::Private::SRenderGridJobList::RemoveRenderStatusColumn()
 	RenderGridJobListWidget->GetHeaderRow()->RemoveColumn(FRenderGridJobListColumns::RenderingStatus);
 }
 
-
 void UE::RenderGrid::Private::SRenderGridJobList::OnObjectModified(UObject* Object)
 {
 	if (Cast<UMoviePipelineOutputSetting>(Object) || Cast<UMovieSceneSequence>(Object) || Cast<UMovieScene>(Object) || Cast<UMovieSceneTrack>(Object) || Cast<UMovieSceneSection>(Object) || Cast<UMovieSceneSubTrack>(Object) || Cast<UMovieSceneSubSection>(Object))
@@ -253,7 +303,6 @@ void UE::RenderGrid::Private::SRenderGridJobList::OnObjectModified(UObject* Obje
 	}
 }
 
-
 void UE::RenderGrid::Private::SRenderGridJobList::Refresh()
 {
 	if (const TSharedPtr<IRenderGridEditor> BlueprintEditor = BlueprintEditorWeakPtr.Pin())
@@ -306,15 +355,39 @@ void UE::RenderGrid::Private::SRenderGridJobList::RefreshHeaderEnabledCheckbox()
 	RenderGridJobEnabledHeaderCheckbox->SetIsChecked(GetDesiredHeaderEnabledCheckboxState());
 }
 
-TSharedRef<ITableRow> UE::RenderGrid::Private::SRenderGridJobList::HandleJobListGenerateRow(URenderGridJob* Item, const TSharedRef<STableViewBase>& OwnerTable)
+
+void UE::RenderGrid::Private::SRenderGridJobListTable::Construct(const FArguments& InArgs, TWeakPtr<IRenderGridEditor> InBlueprintEditor, const TSharedPtr<SRenderGridJobList>& InJobListWidget)
 {
-	TSharedPtr<SWidget> BaseThis = AsShared();
-	TSharedPtr<SRenderGridJobList> This = StaticCastSharedPtr<SRenderGridJobList>(BaseThis);
+	BlueprintEditorWeakPtr = InBlueprintEditor;
+	JobListWidget = InJobListWidget;
+
+	SListView::Construct(
+		SListView::FArguments()
+		.ListItemsSource(InArgs._ListItemsSource)
+		.HeaderRow(InArgs._HeaderRow)
+		.ItemHeight(20.0f)
+		.OnGenerateRow(this, &SRenderGridJobListTable::HandleJobListGenerateRow)
+		.OnSelectionChanged(this, &SRenderGridJobListTable::HandleJobListSelectionChanged)
+		.SelectionMode(ESelectionMode::Multi)
+		.ClearSelectionOnClick(false)
+	);
+}
+
+FReply UE::RenderGrid::Private::SRenderGridJobListTable::OnMouseButtonUp(const FGeometry& MyGeometry, const FPointerEvent& MouseEvent)
+{
+	if (MouseEvent.GetEffectingButton() == EKeys::RightMouseButton)
+	{
+		return JobListWidget->OnMouseButtonUp(MyGeometry, MouseEvent);
+	}
+	return SListView<URenderGridJob*>::OnMouseButtonUp(MyGeometry, MouseEvent);
+}
 
-	return SNew(SRenderGridJobListTableRow, OwnerTable, BlueprintEditorWeakPtr, Item, This);
+TSharedRef<ITableRow> UE::RenderGrid::Private::SRenderGridJobListTable::HandleJobListGenerateRow(URenderGridJob* Item, const TSharedRef<STableViewBase>& OwnerTable)
+{
+	return SNew(SRenderGridJobListTableRow, OwnerTable, BlueprintEditorWeakPtr, Item, JobListWidget);
 }
 
-void UE::RenderGrid::Private::SRenderGridJobList::HandleJobListSelectionChanged(URenderGridJob* Item, ESelectInfo::Type SelectInfo)
+void UE::RenderGrid::Private::SRenderGridJobListTable::HandleJobListSelectionChanged(URenderGridJob* Item, ESelectInfo::Type SelectInfo)
 {
 	if (SelectInfo == ESelectInfo::Type::Direct)
 	{
@@ -322,7 +395,7 @@ void UE::RenderGrid::Private::SRenderGridJobList::HandleJobListSelectionChanged(
 	}
 	if (const TSharedPtr<IRenderGridEditor> BlueprintEditor = BlueprintEditorWeakPtr.Pin())
 	{
-		BlueprintEditor->SetSelectedRenderGridJobs(RenderGridJobListWidget->GetSelectedItems());
+		BlueprintEditor->SetSelectedRenderGridJobs(GetSelectedItems());
 	}
 }
 
@@ -341,38 +414,11 @@ void UE::RenderGrid::Private::SRenderGridJobListTableRow::Construct(const FArgum
 
 FReply UE::RenderGrid::Private::SRenderGridJobListTableRow::OnMouseButtonUp(const FGeometry& MyGeometry, const FPointerEvent& MouseEvent)
 {
-	if (MouseEvent.GetEffectingButton() != EKeys::RightMouseButton)
+	if (MouseEvent.GetEffectingButton() == EKeys::RightMouseButton)
 	{
-		return SMultiColumnTableRow<URenderGridJob*>::OnMouseButtonUp(MyGeometry, MouseEvent);
+		return JobListWidget->OnMouseButtonUp(MyGeometry, MouseEvent);
 	}
-
-	FMenuBuilder MenuBuilder(true, nullptr);
-	{
-		MenuBuilder.BeginSection("RenderGridJobListRowContextMenuSection", LOCTEXT("RenderGridJobListRowContextMenuHeading", "Options"));
-		{
-			MenuBuilder.AddMenuEntry(
-				LOCTEXT("DuplicateRenderGridJobListRowLabel", "Duplicate"),
-				LOCTEXT("DuplicateRenderGridJobListTooltip", "Creates a copy of the job and adds it to the grid."),
-				FSlateIcon(FAppStyle::GetAppStyleSetName(), "Icons.Duplicate"),
-				FUIAction(FExecuteAction::CreateSP(this, &SRenderGridJobListTableRow::DuplicateJob))
-			);
-			MenuBuilder.AddMenuEntry(
-				LOCTEXT("DeleteRenderGridJobListRowLabel", "Delete"),
-				LOCTEXT("DeleteRenderGridJobListTooltip", "Removes the job from the grid."),
-				FSlateIcon(FAppStyle::GetAppStyleSetName(), "Icons.Delete"),
-				FUIAction(FExecuteAction::CreateSP(this, &SRenderGridJobListTableRow::DeleteJob))
-			);
-		}
-		MenuBuilder.EndSection();
-	}
-	TSharedPtr<SWidget> MenuContent = MenuBuilder.MakeWidget();
-
-	if (MenuContent.IsValid() && (MouseEvent.GetEventPath() != nullptr))
-	{
-		FWidgetPath WidgetPath = *MouseEvent.GetEventPath();
-		FSlateApplication::Get().PushMenu(WidgetPath.Widgets.Last().Widget, WidgetPath, MenuContent.ToSharedRef(), MouseEvent.GetScreenSpacePosition(), FPopupTransitionEffect(FPopupTransitionEffect::ContextMenu));
-	}
-	return FReply::Handled();
+	return SMultiColumnTableRow<URenderGridJob*>::OnMouseButtonUp(MyGeometry, MouseEvent);
 }
 
 TOptional<EItemDropZone> UE::RenderGrid::Private::SRenderGridJobListTableRow::OnCanAcceptDrop(const FDragDropEvent& InEvent, EItemDropZone InItemDropZone, URenderGridJob* InJob)
@@ -633,32 +679,6 @@ TSharedRef<SWidget> UE::RenderGrid::Private::SRenderGridJobListTableRow::Generat
 	return SNullWidget::NullWidget;
 }
 
-void UE::RenderGrid::Private::SRenderGridJobListTableRow::DuplicateJob()
-{
-	if (const TSharedPtr<IRenderGridEditor> BlueprintEditor = BlueprintEditorWeakPtr.Pin())
-	{
-		if (URenderGrid* RenderGrid = BlueprintEditor->GetInstance(); IsValid(RenderGrid))
-		{
-			RenderGrid->DuplicateAndAddRenderGridJob(RenderGridJob);
-			BlueprintEditor->MarkAsModified();
-			BlueprintEditor->OnRenderGridChanged().Broadcast();
-		}
-	}
-}
-
-void UE::RenderGrid::Private::SRenderGridJobListTableRow::DeleteJob()
-{
-	if (const TSharedPtr<IRenderGridEditor> BlueprintEditor = BlueprintEditorWeakPtr.Pin())
-	{
-		if (URenderGrid* RenderGrid = BlueprintEditor->GetInstance(); IsValid(RenderGrid))
-		{
-			RenderGrid->RemoveRenderGridJob(RenderGridJob);
-			BlueprintEditor->MarkAsModified();
-			BlueprintEditor->OnRenderGridChanged().Broadcast();
-		}
-	}
-}
-
 FText UE::RenderGrid::Private::SRenderGridJobListTableRow::GetRenderStatusText() const
 {
 	if (const TSharedPtr<IRenderGridEditor> BlueprintEditor = BlueprintEditorWeakPtr.Pin())
diff --git a/Engine/Plugins/Experimental/RenderGrid/Source/RenderGridEditor/Private/UI/SRenderGridJobList.h b/Engine/Plugins/Experimental/RenderGrid/Source/RenderGridEditor/Private/UI/SRenderGridJobList.h
--- a/Engine/Plugins/Experimental/RenderGrid/Source/RenderGridEditor/Private/UI/SRenderGridJobList.h
+++ b/Engine/Plugins/Experimental/RenderGrid/Source/RenderGridEditor/Private/UI/SRenderGridJobList.h
@@ -8,6 +8,7 @@
 #include "Styling/AppStyle.h"
 #include "Widgets/SCompoundWidget.h"
 #include "Widgets/Text/SInlineEditableTextBlock.h"
+#include "Widgets/Views/SListView.h"
 #include "Widgets/Views/STableRow.h"
 #include "Widgets/Views/STableViewBase.h"
 
@@ -22,6 +23,11 @@ namespace UE::RenderGrid
 	class IRenderGridEditor;
 }
 
+namespace UE::RenderGrid::Private
+{
+	class SRenderGridJobListTable;
+}
+
 
 namespace UE::RenderGrid::Private
 {
@@ -40,6 +46,10 @@ namespace UE::RenderGrid::Private
 		virtual void Tick(const FGeometry&, const double, const float) override;
 		void Construct(const FArguments& InArgs, TSharedPtr<IRenderGridEditor> InBlueprintEditor);
 
+		//~ Begin SWidget Interface
+		virtual FReply OnMouseButtonUp(const FGeometry& MyGeometry, const FPointerEvent& MouseEvent) override;
+		//~ End SWidget Interface
+
 		/** Refreshes the content of this widget. */
 		void Refresh();
 
@@ -68,13 +78,6 @@ namespace UE::RenderGrid::Private
 		void OnBatchRenderingFinished(URenderGridQueue* Queue) { Refresh(); }
 		void OnSearchBarTextChanged(const FText& Text) { Refresh(); }
 
-	private:
-		/** Callback for generating a row widget in the session tree view. */
-		TSharedRef<ITableRow> HandleJobListGenerateRow(URenderGridJob* Item, const TSharedRef<STableViewBase>& OwnerTable);
-
-		/** Callback for session tree view selection changes. */
-		void HandleJobListSelectionChanged(URenderGridJob* Item, ESelectInfo::Type SelectInfo);
-
 	private:
 		/** A reference to the blueprint editor that owns the render grid instance. */
 		TWeakPtr<IRenderGridEditor> BlueprintEditorWeakPtr;
@@ -86,7 +89,7 @@ namespace UE::RenderGrid::Private
 		TArray<URenderGridJob*> RenderGridJobs;
 
 		/** The widget that lists the render grid jobs. */
-		TSharedPtr<SListView<URenderGridJob*>> RenderGridJobListWidget;
+		TSharedPtr<SRenderGridJobListTable> RenderGridJobListWidget;
 
 		/** The search bar widget. */
 		TSharedPtr<SSearchBox> RenderGridSearchBox;
@@ -96,6 +99,39 @@ namespace UE::RenderGrid::Private
 	};
 
 
+	/**
+	 * The widget that contains the rows of render grid jobs.
+	 */
+	class SRenderGridJobListTable : public SListView<URenderGridJob*>
+	{
+	public:
+		SLATE_BEGIN_ARGS(SRenderGridJobListTable) {}
+			SLATE_ARGUMENT(const TArray<URenderGridJob*>*, ListItemsSource)
+			SLATE_ARGUMENT(TSharedPtr<SHeaderRow>, HeaderRow)
+		SLATE_END_ARGS()
+
+		void Construct(const FArguments& InArgs, TWeakPtr<IRenderGridEditor> InBlueprintEditor, const TSharedPtr<SRenderGridJobList>& InJobListWidget);
+
+		//~ Begin SWidget Interface
+		virtual FReply OnMouseButtonUp(const FGeometry& MyGeometry, const FPointerEvent& MouseEvent) override;
+		//~ End SWidget Interface
+
+	private:
+		/** Callback for generating a row widget in the session tree view. */
+		TSharedRef<ITableRow> HandleJobListGenerateRow(URenderGridJob* Item, const TSharedRef<STableViewBase>& OwnerTable);
+
+		/** Callback for session tree view selection changes. */
+		void HandleJobListSelectionChanged(URenderGridJob* Item, ESelectInfo::Type SelectInfo);
+
+	private:
+		/** A reference to the blueprint editor that owns the render grid instance. */
+		TWeakPtr<IRenderGridEditor> BlueprintEditorWeakPtr;
+
+		/** A reference to the job list (the parent widget). */
+		TSharedPtr<SRenderGridJobList> JobListWidget;
+	};
+
+
 	/**
 	 * The widget that represents a single render grid job (a single row).
 	 */
@@ -115,9 +151,6 @@ namespace UE::RenderGrid::Private
 		FReply OnAcceptDrop(const FDragDropEvent& InEvent, EItemDropZone InItemDropZone, URenderGridJob* InJob);
 		virtual TSharedRef<SWidget> GenerateWidgetForColumn(const FName& ColumnName) override;
 
-		void DuplicateJob();
-		void DeleteJob();
-
 		FText GetRenderStatusText() const;
 
 	private:
diff --git a/Engine/Plugins/Experimental/RenderGrid/Source/RenderGridEditor/Private/UI/SRenderGridViewerLive.cpp b/Engine/Plugins/Experimental/RenderGrid/Source/RenderGridEditor/Private/UI/SRenderGridViewerLive.cpp
--- a/Engine/Plugins/Experimental/RenderGrid/Source/RenderGridEditor/Private/UI/SRenderGridViewerLive.cpp
+++ b/Engine/Plugins/Experimental/RenderGrid/Source/RenderGridEditor/Private/UI/SRenderGridViewerLive.cpp
@@ -172,7 +172,6 @@ ULevelSequencePlayer* UE::RenderGrid::Private::SRenderGridEditorViewport::GetSeq
 		PlaybackSettings.bAutoPlay = false;
 		PlaybackSettings.bPauseAtEnd = true;
 		PlaybackSettings.bRestoreState = true;
-		FLevelSequenceCameraSettings CameraSettings;
 
 		ALevelSequenceActor* PlayerActor = nullptr;
 		if (ULevelSequencePlayer* Player = ULevelSequencePlayer::CreateLevelSequencePlayer(World, LevelSequence, PlaybackSettings, PlayerActor); IsValid(Player))
