RenderGrid plugin:
- Added a default values panel for render grid jobs.
- Placed the default values in a separate object, to be able to show it in a separate UI.
- Flattened out the categories of all the render grid variables in the UI, in order to make the UI more user-friendly.
- Added undo/redo support.
- Reworked the way the selected jobs are stored in the undo/redo buffer.
- Made the undo/redo buffer no longer be affected by undo/redos from outside the render grid editor.
- Made the render grid plugin store remote control property value changes itself.
- Fixed several garbage collection issues.
- Updated the plugin to work with the latest version of the engine (including with the latest version of the movie pipeline plugin).

diff --git a/Engine/Plugins/Experimental/RenderGrid/Source/RenderGrid/Private/RenderGrid/RenderGrid.cpp b/Engine/Plugins/Experimental/RenderGrid/Source/RenderGrid/Private/RenderGrid/RenderGrid.cpp
--- a/Engine/Plugins/Experimental/RenderGrid/Source/RenderGrid/Private/RenderGrid/RenderGrid.cpp
+++ b/Engine/Plugins/Experimental/RenderGrid/Source/RenderGrid/Private/RenderGrid/RenderGrid.cpp
@@ -12,10 +12,19 @@
 #include "UObject/ObjectSaveContext.h"
 
 
+URenderGridDefaults::URenderGridDefaults()
+	: LevelSequence(nullptr)
+	, RenderPreset(nullptr)
+	, OutputDirectory(UE::RenderGrid::Private::FRenderGridUtils::NormalizeJobOutputDirectory(FPaths::ProjectDir() / TEXT("Saved/MovieRenders/")))
+{
+	SetFlags(RF_Public | RF_Transactional);
+}
+
+
 URenderGridJob::URenderGridJob()
 	: Guid(FGuid::NewGuid())
 	, WaitFramesBeforeRendering(0)
-	, Sequence(nullptr)
+	, LevelSequence(nullptr)
 	, bOverrideStartFrame(false)
 	, CustomStartFrame(0)
 	, bOverrideEndFrame(false)
@@ -24,16 +33,18 @@ URenderGridJob::URenderGridJob()
 	, CustomResolution(FIntPoint(3840, 2160))
 	, bIsEnabled(true)
 	, RenderPreset(nullptr)
-{}
+{
+	SetFlags(RF_Public | RF_Transactional);
+}
 
 TOptional<int32> URenderGridJob::GetSequenceStartFrame() const
 {
-	if (!IsValid(Sequence))
+	if (!IsValid(LevelSequence))
 	{
 		return TOptional<int32>();
 	}
 
-	const UMovieScene* MovieScene = Sequence->MovieScene;
+	const UMovieScene* MovieScene = LevelSequence->MovieScene;
 	const FFrameRate TickResolution = MovieScene->GetTickResolution();
 	const FFrameRate DisplayRate = MovieScene->GetDisplayRate();
 	const UMoviePipelineOutputSetting* Settings = GetRenderPresetOutputSettings();
@@ -54,12 +65,12 @@ TOptional<int32> URenderGridJob::GetSequenceStartFrame() const
 
 TOptional<int32> URenderGridJob::GetSequenceEndFrame() const
 {
-	if (!IsValid(Sequence))
+	if (!IsValid(LevelSequence))
 	{
 		return TOptional<int32>();
 	}
 
-	const UMovieScene* MovieScene = Sequence->MovieScene;
+	const UMovieScene* MovieScene = LevelSequence->MovieScene;
 	const FFrameRate TickResolution = MovieScene->GetTickResolution();
 	const FFrameRate DisplayRate = MovieScene->GetDisplayRate();
 	const UMoviePipelineOutputSetting* Settings = GetRenderPresetOutputSettings();
@@ -151,12 +162,12 @@ TOptional<int32> URenderGridJob::GetStartFrame() const
 		return Settings->CustomStartFrame;
 	}
 
-	if (!IsValid(Sequence))
+	if (!IsValid(LevelSequence))
 	{
 		return TOptional<int32>();
 	}
 
-	const UMovieScene* MovieScene = Sequence->MovieScene;
+	const UMovieScene* MovieScene = LevelSequence->MovieScene;
 	const FFrameRate TickResolution = MovieScene->GetTickResolution();
 	FFrameRate DisplayRate = MovieScene->GetDisplayRate();
 	if (IsValid(Settings) && Settings->bUseCustomFrameRate)
@@ -180,12 +191,12 @@ TOptional<int32> URenderGridJob::GetEndFrame() const
 		return Settings->CustomEndFrame;
 	}
 
-	if (!IsValid(Sequence))
+	if (!IsValid(LevelSequence))
 	{
 		return TOptional<int32>();
 	}
 
-	const UMovieScene* MovieScene = Sequence->MovieScene;
+	const UMovieScene* MovieScene = LevelSequence->MovieScene;
 	const FFrameRate TickResolution = MovieScene->GetTickResolution();
 	FFrameRate DisplayRate = MovieScene->GetDisplayRate();
 	if (IsValid(Settings) && Settings->bUseCustomFrameRate)
@@ -198,7 +209,7 @@ TOptional<int32> URenderGridJob::GetEndFrame() const
 
 TOptional<double> URenderGridJob::GetStartTime() const
 {
-	if (!IsValid(Sequence))
+	if (!IsValid(LevelSequence))
 	{
 		return TOptional<double>();
 	}
@@ -211,7 +222,7 @@ TOptional<double> URenderGridJob::GetStartTime() const
 
 	const UMoviePipelineOutputSetting* Settings = GetRenderPresetOutputSettings();
 
-	const UMovieScene* MovieScene = Sequence->MovieScene;
+	const UMovieScene* MovieScene = LevelSequence->MovieScene;
 	FFrameRate DisplayRate = MovieScene->GetDisplayRate();
 	if (IsValid(Settings) && Settings->bUseCustomFrameRate)
 	{
@@ -222,7 +233,7 @@ TOptional<double> URenderGridJob::GetStartTime() const
 
 TOptional<double> URenderGridJob::GetEndTime() const
 {
-	if (!IsValid(Sequence))
+	if (!IsValid(LevelSequence))
 	{
 		return TOptional<double>();
 	}
@@ -235,7 +246,7 @@ TOptional<double> URenderGridJob::GetEndTime() const
 
 	const UMoviePipelineOutputSetting* Settings = GetRenderPresetOutputSettings();
 
-	const UMovieScene* MovieScene = Sequence->MovieScene;
+	const UMovieScene* MovieScene = LevelSequence->MovieScene;
 	FFrameRate DisplayRate = MovieScene->GetDisplayRate();
 	if (IsValid(Settings) && Settings->bUseCustomFrameRate)
 	{
@@ -246,7 +257,7 @@ TOptional<double> URenderGridJob::GetEndTime() const
 
 TOptional<double> URenderGridJob::GetDurationInSeconds() const
 {
-	if (!IsValid(Sequence))
+	if (!IsValid(LevelSequence))
 	{
 		return TOptional<double>();
 	}
@@ -260,7 +271,7 @@ TOptional<double> URenderGridJob::GetDurationInSeconds() const
 
 	const UMoviePipelineOutputSetting* Settings = GetRenderPresetOutputSettings();
 
-	const UMovieScene* MovieScene = Sequence->MovieScene;
+	const UMovieScene* MovieScene = LevelSequence->MovieScene;
 	FFrameRate DisplayRate = MovieScene->GetDisplayRate();
 	if (IsValid(Settings) && Settings->bUseCustomFrameRate)
 	{
@@ -322,56 +333,19 @@ bool URenderGridJob::MatchesSearchTerm(const FString& SearchTerm) const
 	return true;
 }
 
-FString URenderGridJob::PurgeJobIdOrReturnEmptyString(const FString& NewJobId)
+void URenderGridJob::SetJobId(const FString& NewJobId)
 {
-	static FString ValidCharacters = TEXT("abcdefghijklmnopqrstuvwxyzABCDEFGHIJKLMNOPQRSTUVWXYZ0123456789-_");
-
-	FString Result;
-	for (TCHAR NewJobIdChar : NewJobId)
-	{
-		int32 Index;
-		if (ValidCharacters.FindChar(NewJobIdChar, Index))
-		{
-			Result += NewJobIdChar;
-		}
-	}
-	return Result;
+	JobId = UE::RenderGrid::Private::FRenderGridUtils::PurgeJobId(NewJobId);
 }
 
-FString URenderGridJob::PurgeJobId(const FString& NewJobId)
-{
-	FString Result = PurgeJobIdOrReturnEmptyString(NewJobId);
-	if (Result.IsEmpty())
-	{
-		return TEXT("0");
-	}
-	return Result;
-}
-
-FString URenderGridJob::PurgeJobIdOrGenerateUniqueId(URenderGrid* Grid, const FString& NewJobId)
-{
-	FString Result = PurgeJobIdOrReturnEmptyString(NewJobId);
-	if (Result.IsEmpty())
-	{
-		return Grid->GenerateNextJobId();
-	}
-	return Result;
-}
-
-FString URenderGridJob::PurgeJobName(const FString& NewJobName)
-{
-	return NewJobName.TrimStartAndEnd();
-}
-
-FString URenderGridJob::PurgeOutputDirectory(const FString& NewOutputDirectory)
+FString URenderGridJob::GetOutputDirectory() const
 {
-	return UE::RenderGrid::Private::FRenderGridUtils::NormalizeOutputDirectory(FPaths::ConvertRelativePathToFull(NewOutputDirectory))
-		.Replace(*UE::RenderGrid::Private::FRenderGridUtils::NormalizeOutputDirectory(FPaths::ConvertRelativePathToFull(FPaths::ProjectDir())), TEXT("{project_dir}/"));
+	return UE::RenderGrid::Private::FRenderGridUtils::DenormalizeJobOutputDirectory(OutputDirectory);
 }
 
-FString URenderGridJob::GetOutputDirectory() const
+void URenderGridJob::SetOutputDirectory(const FString& NewOutputDirectory)
 {
-	return UE::RenderGrid::Private::FRenderGridUtils::NormalizeOutputDirectory(FPaths::ConvertRelativePathToFull(OutputDirectory.Replace(TEXT("{project_dir}"), *FPaths::ProjectDir())));
+	OutputDirectory = UE::RenderGrid::Private::FRenderGridUtils::NormalizeJobOutputDirectory(NewOutputDirectory);
 }
 
 UMoviePipelineOutputSetting* URenderGridJob::GetRenderPresetOutputSettings() const
@@ -459,16 +433,18 @@ bool URenderGridJob::SetRemoteControlValue(const TSharedPtr<FRemoteControlEntity
 }
 
 
-URenderGrid::URenderGrid()
-	: Guid(FGuid::NewGuid())
+URenderGrid::URenderGrid(const FObjectInitializer& ObjectInitializer): Super(ObjectInitializer)
+	, Guid(FGuid::NewGuid())
 	, PropsSourceType(ERenderGridPropsSourceType::RemoteControl)
 	, PropsSourceOrigin_RemoteControl(nullptr)
+	, Defaults(ObjectInitializer.CreateDefaultSubobject<URenderGridDefaults>(this, TEXT("RenderGridDefaults")))
 	, bExecutingBlueprintEvent(false)
 	, CachedPropsSource(nullptr)
 	, CachedPropsSourceType(ERenderGridPropsSourceType::Local)
 {
 	if (!HasAnyFlags(RF_ClassDefaultObject))
 	{
+		SetFlags(RF_Public | RF_Transactional);
 		LoadValuesFromCDO();
 		OnPreSaveCDO().AddUObject(this, &URenderGrid::SaveValuesToCDO);
 	}
@@ -532,6 +508,20 @@ void URenderGrid::PostLoad()
 	{
 		PropsSourceType = ERenderGridPropsSourceType::RemoteControl;
 	}
+
+	SetFlags(RF_Public | RF_Transactional);
+	for (TObjectPtr<URenderGridJob> Job : RenderGridJobs)
+	{
+		if (IsValid(Job))
+		{
+			Job->SetFlags(RF_Public | RF_Transactional);
+		}
+	}
+	if (!IsValid(Defaults))
+	{
+		Defaults = NewObject<URenderGridDefaults>(this, TEXT("RenderGridDefaults"));
+	}
+	Defaults->SetFlags(RF_Public | RF_Transactional);
 }
 
 void URenderGrid::BeginEditor()
@@ -641,6 +631,8 @@ void URenderGrid::CopyValuesToOrFromCDO(const bool bToCDO)
 				CDO->RenderGridJobs.Add(JobCopy);
 			}
 		}
+
+		CDO->Defaults = DuplicateObject(Defaults, CDO);
 	}
 }
 
@@ -676,6 +668,40 @@ UObject* URenderGrid::GetPropsSourceOrigin() const
 	return nullptr;
 }
 
+UMoviePipelineOutputSetting* URenderGrid::GetDefaultRenderPresetOutputSettings() const
+{
+	if (!IsValid(Defaults->RenderPreset))
+	{
+		return nullptr;
+	}
+	for (UMoviePipelineSetting* Settings : Defaults->RenderPreset->FindSettingsByClass(UMoviePipelineOutputSetting::StaticClass(), false))
+	{
+		if (!IsValid(Settings))
+		{
+			continue;
+		}
+		if (UMoviePipelineOutputSetting* OutputSettings = Cast<UMoviePipelineOutputSetting>(Settings))
+		{
+			if (!OutputSettings->IsEnabled())
+			{
+				continue;
+			}
+			return OutputSettings;
+		}
+	}
+	return nullptr;
+}
+
+FString URenderGrid::GetDefaultOutputDirectory() const
+{
+	return UE::RenderGrid::Private::FRenderGridUtils::DenormalizeJobOutputDirectory(Defaults->OutputDirectory);
+}
+
+void URenderGrid::SetDefaultOutputDirectory(const FString& NewOutputDirectory)
+{
+	Defaults->OutputDirectory = UE::RenderGrid::Private::FRenderGridUtils::NormalizeJobOutputDirectory(NewOutputDirectory);
+}
+
 void URenderGrid::AddRenderGridJob(URenderGridJob* Job)
 {
 	if (IsValid(Job))
@@ -864,7 +890,9 @@ URenderGridJob* URenderGrid::CreateTempRenderGridJob()
 	URenderGridJob* Job = NewObject<URenderGridJob>(this);
 	Job->SetJobId(GenerateUniqueRandomJobId());
 	Job->SetJobName(TEXT("New"));
-	Job->SetOutputDirectory(FPaths::ProjectDir() / TEXT("Saved/MovieRenders/"));
+	Job->SetLevelSequence(Defaults->LevelSequence);
+	Job->SetRenderPreset(Defaults->RenderPreset);
+	Job->SetOutputDirectoryRaw(Defaults->OutputDirectory);
 
 	if (URenderGridPropsSourceRemoteControl* PropsSource = GetPropsSource<URenderGridPropsSourceRemoteControl>())
 	{
diff --git a/Engine/Plugins/Experimental/RenderGrid/Source/RenderGrid/Private/RenderGrid/RenderGridManager.cpp b/Engine/Plugins/Experimental/RenderGrid/Source/RenderGrid/Private/RenderGrid/RenderGridManager.cpp
--- a/Engine/Plugins/Experimental/RenderGrid/Source/RenderGrid/Private/RenderGrid/RenderGridManager.cpp
+++ b/Engine/Plugins/Experimental/RenderGrid/Source/RenderGrid/Private/RenderGrid/RenderGridManager.cpp
@@ -10,7 +10,7 @@
 URenderGridQueue* UE::RenderGrid::FRenderGridManager::CreateBatchRenderQueue(URenderGrid* Grid)
 {
 	FRenderGridQueueCreateArgs JobArgs;
-	JobArgs.RenderGrid = Grid;
+	JobArgs.RenderGrid = TStrongObjectPtr(Grid);
 	JobArgs.RenderGridJobs.Append(Grid->GetEnabledRenderGridJobs());
 	JobArgs.bIsBatchRender = true;
 	URenderGridQueue* NewRenderQueue = URenderGridQueue::Create(JobArgs);
@@ -26,13 +26,21 @@ URenderGridQueue* UE::RenderGrid::FRenderGridManager::RenderPreviewFrame(const F
 {
 	const FRenderGridManagerRenderPreviewFrameArgsCallback Callback = Args.Callback;
 
-	if (!IsValid(Args.RenderGridJob))
+	URenderGrid* RenderGrid = Args.RenderGrid.Get();
+	if (!IsValid(RenderGrid))
 	{
 		Callback.ExecuteIfBound(false);
 		return nullptr;
 	}
 
-	URenderGridJob* JobCopy = DuplicateObject(Args.RenderGridJob.Get(), Args.RenderGridJob->GetOuter());
+	URenderGridJob* RenderGridJob = Args.RenderGridJob.Get();
+	if (!IsValid(RenderGridJob))
+	{
+		Callback.ExecuteIfBound(false);
+		return nullptr;
+	}
+
+	URenderGridJob* JobCopy = DuplicateObject(RenderGridJob, RenderGridJob->GetOuter());
 	if (!IsValid(JobCopy))
 	{
 		Callback.ExecuteIfBound(false);
@@ -65,8 +73,8 @@ URenderGridQueue* UE::RenderGrid::FRenderGridManager::RenderPreviewFrame(const F
 	JobCopy->SetOutputDirectory(TmpRenderedFramesPath / (Args.Frame.IsSet() ? TEXT("PreviewFrame") : TEXT("PreviewFrames")));
 
 	FRenderGridQueueCreateArgs JobArgs;
-	JobArgs.RenderGrid = Args.RenderGrid;
-	JobArgs.RenderGridJobs.Add(JobCopy);
+	JobArgs.RenderGrid = TStrongObjectPtr(RenderGrid);
+	JobArgs.RenderGridJobs.Add(TStrongObjectPtr(JobCopy));
 	JobArgs.bHeadless = Args.bHeadless;
 	JobArgs.bForceOutputImage = true;
 	JobArgs.bForceOnlySingleOutput = true;
diff --git a/Engine/Plugins/Experimental/RenderGrid/Source/RenderGrid/Private/RenderGrid/RenderGridQueue.cpp b/Engine/Plugins/Experimental/RenderGrid/Source/RenderGrid/Private/RenderGrid/RenderGridQueue.cpp
--- a/Engine/Plugins/Experimental/RenderGrid/Source/RenderGrid/Private/RenderGrid/RenderGridQueue.cpp
+++ b/Engine/Plugins/Experimental/RenderGrid/Source/RenderGrid/Private/RenderGrid/RenderGridQueue.cpp
@@ -18,12 +18,19 @@
 #include "MoviePipelineImageSequenceOutput.h"
 #include "MoviePipelineOutputSetting.h"
 #include "MoviePipelinePIEExecutor.h"
+#include "MoviePipelinePrimaryConfig.h"
 #include "MoviePipelineQueue.h"
 
 
 URenderGridMoviePipelineRenderJob* URenderGridMoviePipelineRenderJob::Create(URenderGridQueue* Queue, URenderGridJob* Job, const UE::RenderGrid::FRenderGridQueueCreateArgs& Args)
 {
-	if (!IsValid(Queue) || !IsValid(Job) || !IsValid(Args.RenderGrid))
+	if (!IsValid(Queue) || !IsValid(Job))
+	{
+		return nullptr;
+	}
+
+	URenderGrid* RenderGrid = Args.RenderGrid.Get();
+	if (!IsValid(RenderGrid))
 	{
 		return nullptr;
 	}
@@ -36,7 +43,7 @@ URenderGridMoviePipelineRenderJob* URenderGridMoviePipelineRenderJob::Create(URe
 
 	URenderGridMoviePipelineRenderJob* RenderJob = NewObject<URenderGridMoviePipelineRenderJob>(Queue);
 	RenderJob->RenderGridJob = Job;
-	RenderJob->RenderGrid = Args.RenderGrid;
+	RenderJob->RenderGrid = RenderGrid;
 	RenderJob->PipelineQueue = NewObject<UMoviePipelineQueue>(RenderJob);
 	RenderJob->PipelineExecutor = NewObject<UMoviePipelineExecutorBase>(RenderJob, PipelineExecutor);
 	RenderJob->PipelineExecutorJob = nullptr;
@@ -57,7 +64,7 @@ URenderGridMoviePipelineRenderJob* URenderGridMoviePipelineRenderJob::Create(URe
 	}
 
 
-	ULevelSequence* JobSequence = Job->GetSequence();
+	ULevelSequence* JobSequence = Job->GetLevelSequence();
 	if (!IsValid(JobSequence) || !Job->GetSequenceStartFrame().IsSet() || !Job->GetSequenceEndFrame().IsSet() || (Job->GetSequenceStartFrame().Get(0) >= Job->GetSequenceEndFrame().Get(0)))
 	{
 		return RenderJob;
@@ -297,7 +304,8 @@ void URenderGridMoviePipelineRenderJob::ExecuteFinished(UMoviePipelineExecutorBa
 
 URenderGridQueue* URenderGridQueue::Create(const UE::RenderGrid::FRenderGridQueueCreateArgs& Args)
 {
-	if (!IsValid(Args.RenderGrid))
+	URenderGrid* RenderGrid = Args.RenderGrid.Get();
+	if (!IsValid(RenderGrid))
 	{
 		return nullptr;
 	}
@@ -305,12 +313,15 @@ URenderGridQueue* URenderGridQueue::Create(const UE::RenderGrid::FRenderGridQueu
 	URenderGridQueue* RenderQueue = NewObject<URenderGridQueue>(GetTransientPackage());
 	RenderQueue->Args = Args;
 	RenderQueue->Queue = MakeShareable(new UE::RenderGrid::Private::FRenderGridGenericExecutionQueue);
-	RenderQueue->RenderGrid = Args.RenderGrid;
+	RenderQueue->RenderGrid = RenderGrid;
 	RenderQueue->bCanceled = false;
 
-	for (const TObjectPtr<URenderGridJob> Job : Args.RenderGridJobs)
+	for (const TStrongObjectPtr<URenderGridJob>& JobPtr : Args.RenderGridJobs)
 	{
-		RenderQueue->AddJob(Job);
+		if (URenderGridJob* Job = JobPtr.Get(); IsValid(Job))
+		{
+			RenderQueue->AddJob(Job);
+		}
 	}
 	return RenderQueue;
 }
diff --git a/Engine/Plugins/Experimental/RenderGrid/Source/RenderGrid/Private/RenderGridUtils.cpp b/Engine/Plugins/Experimental/RenderGrid/Source/RenderGrid/Private/RenderGridUtils.cpp
--- a/Engine/Plugins/Experimental/RenderGrid/Source/RenderGrid/Private/RenderGridUtils.cpp
+++ b/Engine/Plugins/Experimental/RenderGrid/Source/RenderGrid/Private/RenderGridUtils.cpp
@@ -1,17 +1,18 @@
 // Copyright Epic Games, Inc. All Rights Reserved.
 
 #include "RenderGridUtils.h"
-#include "Engine/Engine.h"
+#include "RenderGrid/RenderGrid.h"
 #include "Engine/Texture2D.h"
 #include "HAL/FileManager.h"
 #include "HAL/PlatformFileManager.h"
 #include "IImageWrapper.h"
 #include "IImageWrapperModule.h"
 #include "ImageUtils.h"
+#include "TextureResource.h"
 #include "Editor/EditorPerformanceSettings.h"
+#include "Engine/Engine.h"
 #include "Misc/Paths.h"
 #include "Modules/ModuleManager.h"
-#include "TextureResource.h"
 
 
 bool UE::RenderGrid::Private::FRenderGridUtils::IsImage(const FString& ImagePath)
@@ -253,6 +254,61 @@ FString UE::RenderGrid::Private::FRenderGridUtils::NormalizeOutputDirectory(cons
 }
 
 
+FString UE::RenderGrid::Private::FRenderGridUtils::PurgeJobIdOrReturnEmptyString(const FString& JobId)
+{
+	static FString ValidCharacters = TEXT("abcdefghijklmnopqrstuvwxyzABCDEFGHIJKLMNOPQRSTUVWXYZ0123456789-_");
+
+	FString Result;
+	for (TCHAR JobIdChar : JobId)
+	{
+		int32 Index;
+		if (ValidCharacters.FindChar(JobIdChar, Index))
+		{
+			Result += JobIdChar;
+		}
+	}
+	return Result;
+}
+
+FString UE::RenderGrid::Private::FRenderGridUtils::PurgeJobId(const FString& JobId)
+{
+	FString Result = PurgeJobIdOrReturnEmptyString(JobId);
+	if (Result.IsEmpty())
+	{
+		return TEXT("0");
+	}
+	return Result;
+}
+
+FString UE::RenderGrid::Private::FRenderGridUtils::PurgeJobIdOrGenerateUniqueId(URenderGrid* Grid, const FString& JobId)
+{
+	FString Result = PurgeJobIdOrReturnEmptyString(JobId);
+	if (Result.IsEmpty())
+	{
+		return Grid->GenerateNextJobId();
+	}
+	return Result;
+}
+
+FString UE::RenderGrid::Private::FRenderGridUtils::PurgeJobName(const FString& JobName)
+{
+	return JobName.TrimStartAndEnd();
+}
+
+FString UE::RenderGrid::Private::FRenderGridUtils::NormalizeJobOutputDirectory(const FString& OutputDirectory)
+{
+	return NormalizeOutputDirectory(FPaths::ConvertRelativePathToFull(OutputDirectory))
+		.Replace(*NormalizeOutputDirectory(FPaths::ConvertRelativePathToFull(FPaths::ProjectDir())), TEXT("{project_dir}/"));
+}
+
+FString UE::RenderGrid::Private::FRenderGridUtils::DenormalizeJobOutputDirectory(const FString& NormalizedOutputDirectory)
+{
+	return NormalizeOutputDirectory(FPaths::ConvertRelativePathToFull(
+		NormalizedOutputDirectory.Replace(TEXT("{project_dir}"), *FPaths::ProjectDir())
+	));
+}
+
+
 FRenderGridPreviousEngineFpsSettings UE::RenderGrid::Private::FRenderGridUtils::DisableFpsLimit()
 {
 	FRenderGridPreviousEngineFpsSettings Settings;
diff --git a/Engine/Plugins/Experimental/RenderGrid/Source/RenderGrid/Public/RenderGrid/RenderGrid.h b/Engine/Plugins/Experimental/RenderGrid/Source/RenderGrid/Public/RenderGrid/RenderGrid.h
--- a/Engine/Plugins/Experimental/RenderGrid/Source/RenderGrid/Public/RenderGrid/RenderGrid.h
+++ b/Engine/Plugins/Experimental/RenderGrid/Source/RenderGrid/Public/RenderGrid/RenderGrid.h
@@ -38,6 +38,34 @@ public:
 };
 
 
+/**
+ * This class contains the default values of render grid jobs.
+ *
+ * This is placed in a separate class, primarily so that it be separated easily from the render grid when using the details view widget.
+ */
+UCLASS()
+class RENDERGRID_API URenderGridDefaults : public UObject
+{
+	GENERATED_BODY()
+
+public:
+	URenderGridDefaults();
+
+public:
+	/** The default level sequence for new jobs, this is what will be rendered during rendering. A job without a level sequence can't be rendered. */
+	UPROPERTY(EditInstanceOnly, Category="Defaults", Meta=(DisplayName="Default Level Sequence"))
+	TObjectPtr<ULevelSequence> LevelSequence;
+
+	/** The default movie pipeline render preset for new jobs. Render grid jobs are rendered using the movie pipeline plugin. This 'preset' contains the configuration of that plugin. */
+	UPROPERTY(EditInstanceOnly, Category="Defaults", Meta=(DisplayName="Default Render Preset"))
+	TObjectPtr<UMoviePipelinePrimaryConfig> RenderPreset;
+
+	/** The default output directory for new jobs. This is the folder in which the output files (of rendering) are placed into. To be more specific, the output files are placed in: {OutputDirectory}/{JobId}/, this folder will be created if it doesn't exist at the time of rendering. */
+	UPROPERTY(EditInstanceOnly, Category="Defaults", Meta=(DisplayName="Default Output Directory"))
+	FString OutputDirectory;
+};
+
+
 /**
  * This class represents a render grid job, in other words, an entry (a row) of a render grid.
  * It contains a level sequence and custom properties that will be applied while rendering.
@@ -80,130 +108,122 @@ public:
 	TOptional<double> GetDurationInSeconds() const;
 
 	/** Gets the resolution that this job will be rendered in. */
-	UFUNCTION(BlueprintPure, Category="Render Grid|Job")
+	UFUNCTION(BlueprintPure, Category="Render Grid Job")
 	FIntPoint GetOutputResolution() const;
 
 	/** Gets the aspect ratio that this job will be rendered in. */
-	UFUNCTION(BlueprintPure, Category="Render Grid|Job")
+	UFUNCTION(BlueprintPure, Category="Render Grid Job")
 	double GetOutputAspectRatio() const;
 
 	/** Checks whether the job contains data that matches the search terms. */
-	UFUNCTION(BlueprintPure, Category="Render Grid|Job")
+	UFUNCTION(BlueprintPure, Category="Render Grid Job")
 	bool MatchesSearchTerm(const FString& SearchTerm) const;
 
 public:
 	/** Returns the GUID, which is randomly generated at creation. */
-	UFUNCTION(BlueprintPure, Category="Render Grid|Job")
+	UFUNCTION(BlueprintPure, Category="Render Grid Job")
 	FGuid GetGuid() const { return Guid; }
 
 	/** Randomly generates a new GUID. */
-	UFUNCTION(BlueprintCallable, Category="Render Grid|Job")
+	UFUNCTION(BlueprintCallable, Category="Render Grid Job")
 	void GenerateNewGuid() { Guid = FGuid::NewGuid(); }
 
 
-	UFUNCTION(BlueprintPure, Category="Render Grid|Job")
+	UFUNCTION(BlueprintPure, Category="Render Grid Job")
 	int32 GetWaitFramesBeforeRendering() const { return WaitFramesBeforeRendering; }
 
-	UFUNCTION(BlueprintCallable, Category="Render Grid|Job")
+	UFUNCTION(BlueprintCallable, Category="Render Grid Job")
 	void SetWaitFramesBeforeRendering(const int32 NewWaitFramesBeforeRendering) { WaitFramesBeforeRendering = FMath::Max<int32>(0, NewWaitFramesBeforeRendering); }
 
 
-	UFUNCTION(BlueprintPure, Category="Render Grid|Job")
-	ULevelSequence* GetSequence() const { return Sequence; }
+	UFUNCTION(BlueprintPure, Category="Render Grid Job")
+	ULevelSequence* GetLevelSequence() const { return LevelSequence; }
 
-	UFUNCTION(BlueprintCallable, Category="Render Grid|Job")
-	void SetSequence(ULevelSequence* NewSequence) { Sequence = NewSequence; }
+	UFUNCTION(BlueprintCallable, Category="Render Grid Job")
+	void SetLevelSequence(ULevelSequence* NewSequence) { LevelSequence = NewSequence; }
 
 
-	UFUNCTION(BlueprintPure, Category="Render Grid|Job")
+	UFUNCTION(BlueprintPure, Category="Render Grid Job")
 	bool GetIsUsingCustomStartFrame() const { return bOverrideStartFrame; }
 
-	UFUNCTION(BlueprintCallable, Category="Render Grid|Job")
+	UFUNCTION(BlueprintCallable, Category="Render Grid Job")
 	void SetIsUsingCustomStartFrame(const bool bNewOverrideStartFrame) { bOverrideStartFrame = bNewOverrideStartFrame; }
 
-	UFUNCTION(BlueprintPure, Category="Render Grid|Job")
+	UFUNCTION(BlueprintPure, Category="Render Grid Job")
 	int32 GetCustomStartFrame() const { return CustomStartFrame; }
 
-	UFUNCTION(BlueprintCallable, Category="Render Grid|Job")
+	UFUNCTION(BlueprintCallable, Category="Render Grid Job")
 	void SetCustomStartFrame(const int32 NewCustomStartFrame) { CustomStartFrame = NewCustomStartFrame; }
 
 
-	UFUNCTION(BlueprintPure, Category="Render Grid|Job")
+	UFUNCTION(BlueprintPure, Category="Render Grid Job")
 	bool GetIsUsingCustomEndFrame() const { return bOverrideEndFrame; }
 
-	UFUNCTION(BlueprintCallable, Category="Render Grid|Job")
+	UFUNCTION(BlueprintCallable, Category="Render Grid Job")
 	void SetIsUsingCustomEndFrame(const bool bNewOverrideEndFrame) { bOverrideEndFrame = bNewOverrideEndFrame; }
 
-	UFUNCTION(BlueprintPure, Category="Render Grid|Job")
+	UFUNCTION(BlueprintPure, Category="Render Grid Job")
 	int32 GetCustomEndFrame() const { return CustomEndFrame; }
 
-	UFUNCTION(BlueprintCallable, Category="Render Grid|Job")
+	UFUNCTION(BlueprintCallable, Category="Render Grid Job")
 	void SetCustomEndFrame(const int32 NewCustomEndFrame) { CustomEndFrame = NewCustomEndFrame; }
 
 
-	UFUNCTION(BlueprintPure, Category="Render Grid|Job")
+	UFUNCTION(BlueprintPure, Category="Render Grid Job")
 	bool GetIsUsingCustomResolution() const { return bOverrideResolution; }
 
-	UFUNCTION(BlueprintCallable, Category="Render Grid|Job")
+	UFUNCTION(BlueprintCallable, Category="Render Grid Job")
 	void SetIsUsingCustomResolution(const bool bNewOverrideResolution) { bOverrideResolution = bNewOverrideResolution; }
 
-	UFUNCTION(BlueprintPure, Category="Render Grid|Job")
+	UFUNCTION(BlueprintPure, Category="Render Grid Job")
 	FIntPoint GetCustomResolution() const { return CustomResolution; }
 
-	UFUNCTION(BlueprintCallable, Category="Render Grid|Job")
+	UFUNCTION(BlueprintCallable, Category="Render Grid Job")
 	void SetCustomResolution(const FIntPoint NewCustomResolution) { CustomResolution = NewCustomResolution; }
 
 
-	UFUNCTION(BlueprintPure, Category="Render Grid|Job")
+	UFUNCTION(BlueprintPure, Category="Render Grid Job")
 	FString GetJobId() const { return JobId; }
 
-	static FString PurgeJobIdOrReturnEmptyString(const FString& NewJobId);
-	static FString PurgeJobId(const FString& NewJobId);
-	static FString PurgeJobIdOrGenerateUniqueId(URenderGrid* Grid, const FString& NewJobId);
-
-	UFUNCTION(BlueprintCallable, Category="Render Grid|Job")
-	void SetJobId(const FString& NewJobId) { JobId = PurgeJobId(NewJobId); }
+	UFUNCTION(BlueprintCallable, Category="Render Grid Job")
+	void SetJobId(const FString& NewJobId);
 	void SetJobIdRaw(const FString& NewJobId) { JobId = NewJobId; }
 
 
-	UFUNCTION(BlueprintPure, Category="Render Grid|Job")
+	UFUNCTION(BlueprintPure, Category="Render Grid Job")
 	FString GetJobName() const { return JobName; }
 
-	static FString PurgeJobName(const FString& NewJobName);
-
-	UFUNCTION(BlueprintCallable, Category="Render Grid|Job")
+	UFUNCTION(BlueprintCallable, Category="Render Grid Job")
 	void SetJobName(const FString& NewJobName) { JobName = NewJobName; }
 	void SetJobNameRaw(const FString& NewJobName) { JobName = NewJobName; }
 
 
-	UFUNCTION(BlueprintPure, Category="Render Grid|Job")
+	UFUNCTION(BlueprintPure, Category="Render Grid Job")
 	bool GetIsEnabled() const { return bIsEnabled; }
 
-	UFUNCTION(BlueprintCallable, Category="Render Grid|Job")
+	UFUNCTION(BlueprintCallable, Category="Render Grid Job")
 	void SetIsEnabled(const bool bEnabled) { bIsEnabled = bEnabled; }
 
 
-	UFUNCTION(BlueprintPure, Category="Render Grid|Job")
+	UFUNCTION(BlueprintPure, Category="Render Grid Job")
 	FString GetOutputDirectory() const;
 	FString GetOutputDirectoryRaw() const { return OutputDirectory; }
 
-	UFUNCTION(BlueprintPure, Category="Render Grid|Job")
+	UFUNCTION(BlueprintPure, Category="Render Grid Job")
 	FString GetOutputDirectoryForDisplay() const { return OutputDirectory; }
 
-	static FString PurgeOutputDirectory(const FString& NewOutputDirectory);
-
-	UFUNCTION(BlueprintCallable, Category="Render Grid|Job")
-	void SetOutputDirectory(const FString& NewOutputDirectory) { OutputDirectory = PurgeOutputDirectory(NewOutputDirectory); }
+	UFUNCTION(BlueprintCallable, Category="Render Grid Job")
+	void SetOutputDirectory(const FString& NewOutputDirectory);
 	void SetOutputDirectoryRaw(const FString& NewOutputDirectory) { OutputDirectory = NewOutputDirectory; }
 
 
-	UFUNCTION(BlueprintPure, Category="Render Grid|Job")
+	UFUNCTION(BlueprintPure, Category="Render Grid Job")
 	UMoviePipelinePrimaryConfig* GetRenderPreset() const { return RenderPreset; }
 
-	UFUNCTION(BlueprintPure, Category="Render Grid|Job")
+	UFUNCTION(BlueprintPure, Category="Render Grid Job")
 	UMoviePipelineOutputSetting* GetRenderPresetOutputSettings() const;
 
-	UFUNCTION(BlueprintCallable, Category="Render Grid|Job")
+	UFUNCTION(BlueprintCallable, Category="Render Grid Job")
 	void SetRenderPreset(UMoviePipelinePrimaryConfig* NewRenderPreset) { RenderPreset = NewRenderPreset; }
 
 
@@ -219,35 +239,35 @@ private:
 	FGuid Guid;
 
 	/** Waits the given number of frames before it will render this job. This can be set to a higher amount when the renderer has to wait for your code to complete (such as construction scripts etc). Try increasing this value when rendering doesn't produce the output you expect it to. */
-	UPROPERTY(EditInstanceOnly, Category="Render Grid|Job", Meta=(AllowPrivateAccess="true", ClampMin="0"))
+	UPROPERTY(EditInstanceOnly, Category="Render Grid Job", Meta=(AllowPrivateAccess="true", ClampMin="0"))
 	int32 WaitFramesBeforeRendering;
 
 	/** The level sequence, this is what will be rendered during rendering. A job without a level sequence can't be rendered. */
-	UPROPERTY(EditInstanceOnly, Category="Render Grid|Job", Meta=(AllowPrivateAccess="true"))
-	TObjectPtr<ULevelSequence> Sequence;
+	UPROPERTY(EditInstanceOnly, Category="Render Grid Job", Meta=(AllowPrivateAccess="true"))
+	TObjectPtr<ULevelSequence> LevelSequence;
 
 	/** If this is true, the CustomStartFrame property will override the start frame of the level sequence. */
-	UPROPERTY(EditInstanceOnly, Category="Render Grid|Job", Meta=(AllowPrivateAccess="true", InlineEditConditionToggle))
+	UPROPERTY(EditInstanceOnly, Category="Render Grid Job", Meta=(AllowPrivateAccess="true", InlineEditConditionToggle))
 	bool bOverrideStartFrame;
 
 	/** If bOverrideStartFrame is true, this property will override the start frame of the level sequence. */
-	UPROPERTY(EditInstanceOnly, Category="Render Grid|Job", Meta=(AllowPrivateAccess="true", EditCondition="bOverrideStartFrame"))
+	UPROPERTY(EditInstanceOnly, Category="Render Grid Job", Meta=(AllowPrivateAccess="true", EditCondition="bOverrideStartFrame"))
 	int32 CustomStartFrame;
 
 	/** If this is true, the CustomEndFrame property will override the end frame of the level sequence. */
-	UPROPERTY(EditInstanceOnly, Category="Render Grid|Job", Meta=(AllowPrivateAccess="true", InlineEditConditionToggle))
+	UPROPERTY(EditInstanceOnly, Category="Render Grid Job", Meta=(AllowPrivateAccess="true", InlineEditConditionToggle))
 	bool bOverrideEndFrame;
 
 	/** If bOverrideEndFrame is true, this property will override the end frame of the level sequence. */
-	UPROPERTY(EditInstanceOnly, Category="Render Grid|Job", Meta=(AllowPrivateAccess="true", EditCondition="bOverrideEndFrame"))
+	UPROPERTY(EditInstanceOnly, Category="Render Grid Job", Meta=(AllowPrivateAccess="true", EditCondition="bOverrideEndFrame"))
 	int32 CustomEndFrame;
 
 	/** If this is true, the CustomResolution property will override the resolution of the render. */
-	UPROPERTY(EditInstanceOnly, Category="Render Grid|Job", Meta=(AllowPrivateAccess="true", InlineEditConditionToggle))
+	UPROPERTY(EditInstanceOnly, Category="Render Grid Job", Meta=(AllowPrivateAccess="true", InlineEditConditionToggle))
 	bool bOverrideResolution;
 
 	/** If bOverrideResolution is true, this property will override the resolution of the render. */
-	UPROPERTY(EditInstanceOnly, Category="Render Grid|Job", Meta=(AllowPrivateAccess="true", EditCondition="bOverrideResolution", ClampMin="1", UIMin="1"))
+	UPROPERTY(EditInstanceOnly, Category="Render Grid Job", Meta=(AllowPrivateAccess="true", EditCondition="bOverrideResolution", ClampMin="1", UIMin="1"))
 	FIntPoint CustomResolution;
 
 	/** If this is true, this job will be rendered during a batch rendering, otherwise it will be skipped. */
@@ -286,7 +306,7 @@ class RENDERGRID_API URenderGrid : public UObject
 	GENERATED_BODY()
 
 public:
-	URenderGrid();
+	URenderGrid(const FObjectInitializer& ObjectInitializer);
 
 	//UObject interface
 	virtual UWorld* GetWorld() const override;
@@ -439,17 +459,17 @@ private:
 
 public:
 	/** Returns the GUID, which is randomly generated at creation. */
-	UFUNCTION(BlueprintPure, Category="Render Grid|Grid")
+	UFUNCTION(BlueprintPure, Category="Render Grid")
 	FGuid GetGuid() const { return Guid; }
 
 	/** Randomly generates a new GUID. */
-	UFUNCTION(BlueprintCallable, Category="Render Grid|Grid")
+	UFUNCTION(BlueprintCallable, Category="Render Grid")
 	void GenerateNewGuid() { Guid = FGuid::NewGuid(); }
 
-	UFUNCTION(BlueprintCallable, Category="Render Grid|Grid")
+	UFUNCTION(BlueprintCallable, Category="Render Grid")
 	void SetPropsSource(ERenderGridPropsSourceType InPropsSourceType, UObject* InPropsSourceOrigin = nullptr);
 
-	UFUNCTION(BlueprintPure, Category="Render Grid|Grid")
+	UFUNCTION(BlueprintPure, Category="Render Grid")
 	URenderGridPropsSourceBase* GetPropsSource() const;
 
 	template<typename Type>
@@ -459,75 +479,104 @@ public:
 		return Cast<Type>(GetPropsSource());
 	}
 
-	UFUNCTION(BlueprintPure, Category="Render Grid|Grid")
+	UFUNCTION(BlueprintPure, Category="Render Grid")
 	ERenderGridPropsSourceType GetPropsSourceType() const { return PropsSourceType; }
 
-	UFUNCTION(BlueprintPure, Category="Render Grid|Grid")
+	UFUNCTION(BlueprintPure, Category="Render Grid")
 	UObject* GetPropsSourceOrigin() const;
 
 public:
-	UFUNCTION(BlueprintCallable, Category="Render Grid|Grid")
+	URenderGridDefaults* GetDefaultsObject() { return Defaults; }
+
+	UFUNCTION(BlueprintPure, Category="Render Grid")
+	ULevelSequence* GetDefaultLevelSequence() const { return Defaults->LevelSequence; }
+
+	UFUNCTION(BlueprintCallable, Category="Render Grid")
+	void SetDefaultLevelSequence(ULevelSequence* NewSequence) { Defaults->LevelSequence = NewSequence; }
+
+	UFUNCTION(BlueprintPure, Category="Render Grid")
+	UMoviePipelinePrimaryConfig* GetDefaultRenderPreset() const { return Defaults->RenderPreset; }
+
+	UFUNCTION(BlueprintPure, Category="Render Grid")
+	UMoviePipelineOutputSetting* GetDefaultRenderPresetOutputSettings() const;
+
+	UFUNCTION(BlueprintCallable, Category="Render Grid")
+	void SetDefaultRenderPreset(UMoviePipelinePrimaryConfig* NewRenderPreset) { Defaults->RenderPreset = NewRenderPreset; }
+
+	UFUNCTION(BlueprintPure, Category="Render Grid")
+	FString GetDefaultOutputDirectory() const;
+	FString GetDefaultOutputDirectoryRaw() const { return Defaults->OutputDirectory; }
+
+	UFUNCTION(BlueprintPure, Category="Render Grid")
+	FString GetDefaultOutputDirectoryForDisplay() const { return Defaults->OutputDirectory; }
+
+	UFUNCTION(BlueprintCallable, Category="Render Grid")
+	void SetDefaultOutputDirectory(const FString& NewOutputDirectory);
+	void SetDefaultOutputDirectoryRaw(const FString& NewOutputDirectory) { Defaults->OutputDirectory = NewOutputDirectory; }
+
+public:
+	UFUNCTION(BlueprintCallable, Category="Render Grid")
 	void AddRenderGridJob(URenderGridJob* Job);
 
-	UFUNCTION(BlueprintCallable, Category="Render Grid|Grid")
+	UFUNCTION(BlueprintCallable, Category="Render Grid")
 	void RemoveRenderGridJob(URenderGridJob* Job);
 
-	UFUNCTION(BlueprintCallable, Category="Render Grid|Grid")
+	UFUNCTION(BlueprintCallable, Category="Render Grid")
 	void InsertRenderGridJob(URenderGridJob* Job, int32 Index);
 
-	UFUNCTION(BlueprintCallable, Category="Render Grid|Grid")
+	UFUNCTION(BlueprintCallable, Category="Render Grid")
 	bool HasAnyRenderGridJobs() const { return !RenderGridJobs.IsEmpty(); }
 
-	UFUNCTION(BlueprintPure, Category="Render Grid|Grid")
+	UFUNCTION(BlueprintPure, Category="Render Grid")
 	bool HasRenderGridJob(URenderGridJob* Job) const;
 
-	UFUNCTION(BlueprintPure, Category="Render Grid|Grid")
+	UFUNCTION(BlueprintPure, Category="Render Grid")
 	int32 GetIndexOfRenderGridJob(URenderGridJob* Job) const;
 
 	TArray<TObjectPtr<URenderGridJob>>& GetRenderGridJobsRef() { return RenderGridJobs; }
 
-	UFUNCTION(BlueprintPure, Category="Render Grid|Grid")
+	UFUNCTION(BlueprintPure, Category="Render Grid")
 	TArray<URenderGridJob*> GetRenderGridJobs() const;
 
-	UFUNCTION(BlueprintPure, Category="Render Grid|Grid")
+	UFUNCTION(BlueprintPure, Category="Render Grid")
 	TArray<URenderGridJob*> GetEnabledRenderGridJobs() const;
 
-	UFUNCTION(BlueprintPure, Category="Render Grid|Grid")
+	UFUNCTION(BlueprintPure, Category="Render Grid")
 	TArray<URenderGridJob*> GetDisabledRenderGridJobs() const;
 
-	UFUNCTION(BlueprintCallable, Category="Render Grid|Grid")
+	UFUNCTION(BlueprintCallable, Category="Render Grid")
 	void InsertRenderGridJobBefore(URenderGridJob* Job, URenderGridJob* BeforeJob);
 
-	UFUNCTION(BlueprintCallable, Category="Render Grid|Grid")
+	UFUNCTION(BlueprintCallable, Category="Render Grid")
 	void InsertRenderGridJobAfter(URenderGridJob* Job, URenderGridJob* AfterJob);
 
 public:
 	/** Generates a unique job ID by grabbing the current time, as well as 16 random bytes, and converting that to a base64 string. **/
-	UFUNCTION(BlueprintCallable, Category="Render Grid|Grid")
+	UFUNCTION(BlueprintCallable, Category="Render Grid")
 	FString GenerateUniqueRandomJobId();
 
 	/** Generates a unique job ID by finding the currently highest job ID and increasing it by one. **/
-	UFUNCTION(BlueprintCallable, Category="Render Grid|Grid")
+	UFUNCTION(BlueprintCallable, Category="Render Grid")
 	FString GenerateNextJobId();
 
 	/** Finds whether given job ID already exists in this grid. **/
-	UFUNCTION(BlueprintPure, Category="Render Grid|Grid")
+	UFUNCTION(BlueprintPure, Category="Render Grid")
 	bool DoesJobIdExist(const FString& JobId);
 
 	/** Creates a new job. This job won't be added to the grid, so it will eventually be garbage collected. **/
-	UFUNCTION(BlueprintCallable, Category="Render Grid|Grid")
+	UFUNCTION(BlueprintCallable, Category="Render Grid")
 	URenderGridJob* CreateTempRenderGridJob();
 
 	/** Creates a new job and adds it to this grid. **/
-	UFUNCTION(BlueprintCallable, Category="Render Grid|Grid")
+	UFUNCTION(BlueprintCallable, Category="Render Grid")
 	URenderGridJob* CreateAndAddNewRenderGridJob();
 
 	/** Copy the given job in this grid. **/
-	UFUNCTION(BlueprintCallable, Category="Render Grid|Grid")
+	UFUNCTION(BlueprintCallable, Category="Render Grid")
 	URenderGridJob* DuplicateAndAddRenderGridJob(URenderGridJob* Job);
 
 	/** Relocates the given job in this grid to the position of the given dropped-on job. **/
-	UFUNCTION(BlueprintCallable, Category="Render Grid|Grid")
+	UFUNCTION(BlueprintCallable, Category="Render Grid")
 	bool ReorderRenderGridJob(URenderGridJob* Job, URenderGridJob* DroppedOnJob, const bool bAfter = true);
 
 public:
@@ -549,14 +598,19 @@ private:
 
 
 	/** The type of the properties that a job in this grid can have. */
-	UPROPERTY(/*EditInstanceOnly, Category="Render Grid|Source", Meta=(DisplayName="Properties Type", AllowPrivateAccess="true")*/)
+	UPROPERTY(/*EditInstanceOnly, Category="Render Grid", Meta=(DisplayName="Properties Type", AllowPrivateAccess="true")*/)
 	ERenderGridPropsSourceType PropsSourceType;
 
 	/** The remote control properties that a job in this grid can have, only use this if PropsSourceType is ERenderGridPropsSourceType::RemoteControl. */
-	UPROPERTY(EditInstanceOnly, Category="Render Grid|Source", Meta=(DisplayName="Remote Control Preset", AllowPrivateAccess="true", EditCondition="PropsSourceType == ERenderGridPropsSourceType::RemoteControl", EditConditionHides))
+	UPROPERTY(EditInstanceOnly, Category="Render Grid", Meta=(DisplayName="Remote Control Preset", AllowPrivateAccess="true", EditCondition="PropsSourceType == ERenderGridPropsSourceType::RemoteControl", EditConditionHides))
 	TObjectPtr<URemoteControlPreset> PropsSourceOrigin_RemoteControl;
 
 
+	/** The default values for new jobs. */
+	UPROPERTY(Instanced)
+	TObjectPtr<URenderGridDefaults> Defaults;
+
+
 	/** The jobs of this grid. */
 	UPROPERTY(Instanced)
 	TArray<TObjectPtr<URenderGridJob>> RenderGridJobs;
diff --git a/Engine/Plugins/Experimental/RenderGrid/Source/RenderGrid/Public/RenderGrid/RenderGridManager.h b/Engine/Plugins/Experimental/RenderGrid/Source/RenderGrid/Public/RenderGrid/RenderGridManager.h
--- a/Engine/Plugins/Experimental/RenderGrid/Source/RenderGrid/Public/RenderGrid/RenderGridManager.h
+++ b/Engine/Plugins/Experimental/RenderGrid/Source/RenderGrid/Public/RenderGrid/RenderGridManager.h
@@ -3,6 +3,7 @@
 #pragma once
 
 #include "CoreMinimal.h"
+#include "UObject/StrongObjectPtr.h"
 #include "RenderGrid/RenderGrid.h"
 #include "RenderGridManager.generated.h"
 
@@ -54,10 +55,10 @@ namespace UE::RenderGrid
 		bool bHeadless = false;
 
 		/** The render grid of the given render grid jobs that will be rendered. */
-		TObjectPtr<URenderGrid> RenderGrid = nullptr;
+		TStrongObjectPtr<URenderGrid> RenderGrid = nullptr;
 
 		/** The specific render grid job that will be rendered. */
-		TObjectPtr<URenderGridJob> RenderGridJob = nullptr;
+		TStrongObjectPtr<URenderGridJob> RenderGridJob = nullptr;
 
 		/** The specific frame number that will be rendered. */
 		TOptional<int32> Frame;
@@ -65,9 +66,6 @@ namespace UE::RenderGrid
 		/** The resolution it will be rendered in. */
 		FIntPoint Resolution = FIntPoint(0, 0);
 
-		/** The texture to reuse for rendering (performance optimization, prevents a new UTexture2D from having to be created, will only be used if the resolution of this texture matches the resolution it will be rendering in). */
-		TObjectPtr<UTexture2D> ReusingTexture2D = nullptr;
-
 		/** The delegate for when the rendering has finished. */
 		FRenderGridManagerRenderPreviewFrameArgsCallback Callback;
 	};
diff --git a/Engine/Plugins/Experimental/RenderGrid/Source/RenderGrid/Public/RenderGrid/RenderGridQueue.h b/Engine/Plugins/Experimental/RenderGrid/Source/RenderGrid/Public/RenderGrid/RenderGridQueue.h
--- a/Engine/Plugins/Experimental/RenderGrid/Source/RenderGrid/Public/RenderGrid/RenderGridQueue.h
+++ b/Engine/Plugins/Experimental/RenderGrid/Source/RenderGrid/Public/RenderGrid/RenderGridQueue.h
@@ -39,10 +39,10 @@ namespace UE::RenderGrid
 	{
 	public:
 		/** The render grid of the given render grid jobs that will be rendered. */
-		TObjectPtr<URenderGrid> RenderGrid = nullptr;
+		TStrongObjectPtr<URenderGrid> RenderGrid = nullptr;
 
 		/** The specific render grid jobs that will be rendered. */
-		TArray<TObjectPtr<URenderGridJob>> RenderGridJobs;
+		TArray<TStrongObjectPtr<URenderGridJob>> RenderGridJobs;
 
 		/** If not null, it will override the movie pipeline executor class with this class. */
 		TSubclassOf<UMoviePipelinePIEExecutor> PipelineExecutorClass = nullptr;
@@ -167,51 +167,51 @@ public:
 
 public:
 	/** Queues the given job. */
-	UFUNCTION(BlueprintCallable, Category="Render Grid|Queue", Meta=(Keywords="render append"))
+	UFUNCTION(BlueprintCallable, Category="Render Grid Queue", Meta=(Keywords="render append"))
 	void AddJob(URenderGridJob* Job);
 
 	/** Pauses the queue. */
-	UFUNCTION(BlueprintCallable, Category="Render Grid|Queue", Meta=(Keywords="wait"))
+	UFUNCTION(BlueprintCallable, Category="Render Grid Queue", Meta=(Keywords="wait"))
 	void Pause();
 
 	/** Resumes the queue. */
-	UFUNCTION(BlueprintCallable, Category="Render Grid|Queue", Meta=(Keywords="unwait unpause"))
+	UFUNCTION(BlueprintCallable, Category="Render Grid Queue", Meta=(Keywords="unwait unpause"))
 	void Resume();
 
 	/** Cancels the current and the remaining queued jobs. Relies on the internal movie pipeline implementation of job canceling on whether this will stop the current render grid job from rendering or not. Will always prevent new render grid jobs from rendering. */
-	UFUNCTION(BlueprintCallable, Category="Render Grid|Queue", Meta=(Keywords="stop quit exit kill terminate end"))
+	UFUNCTION(BlueprintCallable, Category="Render Grid Queue", Meta=(Keywords="stop quit exit kill terminate end"))
 	void Cancel();
 
 	/** Returns true if this queue has been canceled. */
-	UFUNCTION(BlueprintCallable, Category="Render Grid|Queue", Meta=(Keywords="stopped quited exited killed terminated ended"))
+	UFUNCTION(BlueprintCallable, Category="Render Grid Queue", Meta=(Keywords="stopped quited exited killed terminated ended"))
 	bool IsCanceled() const { return bCanceled; }
 
 	/** Retrieves the rendering status of the given render grid job. */
-	UFUNCTION(BlueprintCallable, Category="Render Grid|Queue", Meta=(Keywords="rendering progression obtain text"))
+	UFUNCTION(BlueprintCallable, Category="Render Grid Queue", Meta=(Keywords="rendering progression obtain text"))
 	FString GetJobStatus(URenderGridJob* Job) const;
 
 	/** Returns all the jobs that have been and will be rendered. */
-	UFUNCTION(BlueprintCallable, Category="Render Grid|Queue", Meta=(Keywords="rendering progression obtain"))
+	UFUNCTION(BlueprintCallable, Category="Render Grid Queue", Meta=(Keywords="rendering progression obtain"))
 	TArray<URenderGridJob*> GetJobs() const;
 
 	/** Returns the number of jobs that have been and will be rendered. */
-	UFUNCTION(BlueprintCallable, Category="Render Grid|Queue", Meta=(Keywords="rendering progression obtain number amount"))
+	UFUNCTION(BlueprintCallable, Category="Render Grid Queue", Meta=(Keywords="rendering progression obtain number amount"))
 	int32 GetJobsCount() const;
 
 	/** Returns the number of jobs that have finished rendering. Basically just returns [Get Jobs Count] minus [Get Jobs Remaining Count]. */
-	UFUNCTION(BlueprintCallable, Category="Render Grid|Queue", Meta=(Keywords="rendering progression obtain number amount finished"))
+	UFUNCTION(BlueprintCallable, Category="Render Grid Queue", Meta=(Keywords="rendering progression obtain number amount finished"))
 	int32 GetJobsCompletedCount() const;
 
 	/** Returns the percentage of jobs finished, this includes the progression of the job that is currently rendering. */
-	UFUNCTION(BlueprintCallable, Category="Render Grid|Queue", Meta=(Keywords="rendering progression obtain number amount finished"))
+	UFUNCTION(BlueprintCallable, Category="Render Grid Queue", Meta=(Keywords="rendering progression obtain number amount finished"))
 	float GetStatusPercentage() const;
 
 	/** Returns the number of jobs that are still left to render, includes the job that is currently rendering. */
-	UFUNCTION(BlueprintCallable, Category="Render Grid|Queue", Meta=(Keywords="rendering progression obtain number amount finished"))
+	UFUNCTION(BlueprintCallable, Category="Render Grid Queue", Meta=(Keywords="rendering progression obtain number amount finished"))
 	int32 GetJobsRemainingCount() const;
 
 	/** Returns the status of the rendering process. */
-	UFUNCTION(BlueprintCallable, Category="Render Grid|Queue", Meta=(Keywords="rendering progression obtain text"))
+	UFUNCTION(BlueprintCallable, Category="Render Grid Queue", Meta=(Keywords="rendering progression obtain text"))
 	FString GetStatus() const;
 
 protected:
diff --git a/Engine/Plugins/Experimental/RenderGrid/Source/RenderGrid/Public/RenderGridUtils.h b/Engine/Plugins/Experimental/RenderGrid/Source/RenderGrid/Public/RenderGridUtils.h
--- a/Engine/Plugins/Experimental/RenderGrid/Source/RenderGrid/Public/RenderGridUtils.h
+++ b/Engine/Plugins/Experimental/RenderGrid/Source/RenderGrid/Public/RenderGridUtils.h
@@ -7,6 +7,7 @@
 #include "RenderGridUtils.generated.h"
 
 
+class URenderGrid;
 class UTexture2D;
 
 
@@ -54,7 +55,7 @@ namespace UE::RenderGrid::Private
 	/**
 	 * A class containing static utility functions for the RenderGrid module.
 	 */
-	class FRenderGridUtils
+	class RENDERGRID_API FRenderGridUtils
 	{
 	public:
 		/**
@@ -135,6 +136,38 @@ namespace UE::RenderGrid::Private
 		static FString NormalizeOutputDirectory(const FString& OutputDirectory);
 
 
+		/**
+		 * Returns a job ID with all invalid characters removed from it. Returns an empty string if it contains no valid characters.
+		 */
+		static FString PurgeJobIdOrReturnEmptyString(const FString& JobId);
+
+		/**
+		 * Returns a job ID with all invalid characters removed from it. Returns string "0" if it contains no valid characters.
+		 */
+		static FString PurgeJobId(const FString& JobId);
+
+		/**
+		 * Returns a job ID with all invalid characters removed from it. Returns a newly generated unique ID if it contains no valid characters.
+		 * The generated string is only unique relative to the job IDs that currently exist in the given RenderGrid.
+		 */
+		static FString PurgeJobIdOrGenerateUniqueId(URenderGrid* Grid, const FString& JobId);
+
+		/**
+		 * Returns a job name with all invalid characters removed from it. Returns an empty string if it contains no valid characters.
+		 */
+		static FString PurgeJobName(const FString& JobName);
+
+		/**
+		 * Returns a job output directory that's converted to a full path (if it's currently relative) and that also has the project directory path replaced by "{project_dir}" (if the given path contains that string).
+		 */
+		static FString NormalizeJobOutputDirectory(const FString& NewOutputDirectory);
+
+		/**
+		 * Returns a job output directory that's a usable valid filesystem path. Replaces string "{project_dir}" by the project directory.
+		 */
+		static FString DenormalizeJobOutputDirectory(const FString& NormalizedOutputDirectory);
+
+
 		/**
 		 * Disables the current FPS limiting, returns the previous settings which allow you to revert back to the previous state.
 		 */
diff --git a/Engine/Plugins/Experimental/RenderGrid/Source/RenderGrid/RenderGrid.Build.cs b/Engine/Plugins/Experimental/RenderGrid/Source/RenderGrid/RenderGrid.Build.cs
--- a/Engine/Plugins/Experimental/RenderGrid/Source/RenderGrid/RenderGrid.Build.cs
+++ b/Engine/Plugins/Experimental/RenderGrid/Source/RenderGrid/RenderGrid.Build.cs
@@ -7,6 +7,8 @@ public class RenderGrid : ModuleRules
 	public RenderGrid(ReadOnlyTargetRules Target) : base(Target)
 	{
 		PCHUsage = PCHUsageMode.UseExplicitOrSharedPCHs;
+		//PCHUsage = PCHUsageMode.NoPCHs;
+		//bUseUnity = false;
 
 		PublicDependencyModuleNames.AddRange(
 			new string[]
@@ -18,7 +20,6 @@ public class RenderGrid : ModuleRules
 			}
 		);
 
-
 		PrivateDependencyModuleNames.AddRange(
 			new string[]
 			{
diff --git a/Engine/Plugins/Experimental/RenderGrid/Source/RenderGridDeveloper/Public/Blueprints/RenderGridBlueprint.h b/Engine/Plugins/Experimental/RenderGrid/Source/RenderGridDeveloper/Public/Blueprints/RenderGridBlueprint.h
--- a/Engine/Plugins/Experimental/RenderGrid/Source/RenderGridDeveloper/Public/Blueprints/RenderGridBlueprint.h
+++ b/Engine/Plugins/Experimental/RenderGrid/Source/RenderGridDeveloper/Public/Blueprints/RenderGridBlueprint.h
@@ -28,7 +28,7 @@ public:
 	virtual bool SupportsDelegates() const override { return true; }
 	virtual bool SupportsEventGraphs() const override { return true; }
 	virtual bool SupportsAnimLayers() const override { return false; }
-
+	virtual bool ShouldBeMarkedDirtyUponTransaction() const override { return false; }
 	virtual UClass* GetBlueprintClass() const override;
 	virtual void PostLoad() override;
 	//~ End UBlueprint Interface
diff --git a/Engine/Plugins/Experimental/RenderGrid/Source/RenderGridDeveloper/RenderGridDeveloper.Build.cs b/Engine/Plugins/Experimental/RenderGrid/Source/RenderGridDeveloper/RenderGridDeveloper.Build.cs
--- a/Engine/Plugins/Experimental/RenderGrid/Source/RenderGridDeveloper/RenderGridDeveloper.Build.cs
+++ b/Engine/Plugins/Experimental/RenderGrid/Source/RenderGridDeveloper/RenderGridDeveloper.Build.cs
@@ -7,6 +7,8 @@ public class RenderGridDeveloper : ModuleRules
 	public RenderGridDeveloper(ReadOnlyTargetRules Target) : base(Target)
 	{
 		PCHUsage = PCHUsageMode.UseExplicitOrSharedPCHs;
+		//PCHUsage = PCHUsageMode.NoPCHs;
+		//bUseUnity = false;
 
 		PublicDependencyModuleNames.AddRange(
 			new string[]
diff --git a/Engine/Plugins/Experimental/RenderGrid/Source/RenderGridEditor/Private/Factories/RenderGridFactory.h b/Engine/Plugins/Experimental/RenderGrid/Source/RenderGridEditor/Private/Factories/RenderGridFactory.h
--- a/Engine/Plugins/Experimental/RenderGrid/Source/RenderGridEditor/Private/Factories/RenderGridFactory.h
+++ b/Engine/Plugins/Experimental/RenderGrid/Source/RenderGridEditor/Private/Factories/RenderGridFactory.h
@@ -29,6 +29,6 @@ public:
 	virtual uint32 GetMenuCategories() const override;
 
 	/** The parent class of the created blueprint. */
-	UPROPERTY(EditAnywhere, Category="Render Grid|Render Grid Factory", Meta = (AllowAbstract = ""))
+	UPROPERTY(EditAnywhere, Category="Render Grid Factory", Meta = (AllowAbstract = ""))
 	TSubclassOf<URenderGrid> ParentClass;
 };
diff --git a/Engine/Plugins/Experimental/RenderGrid/Source/RenderGridEditor/Private/Toolkit/RenderGridEditor.cpp b/Engine/Plugins/Experimental/RenderGrid/Source/RenderGridEditor/Private/Toolkit/RenderGridEditor.cpp
--- a/Engine/Plugins/Experimental/RenderGrid/Source/RenderGridEditor/Private/Toolkit/RenderGridEditor.cpp
+++ b/Engine/Plugins/Experimental/RenderGrid/Source/RenderGridEditor/Private/Toolkit/RenderGridEditor.cpp
@@ -24,6 +24,7 @@
 #include "Misc/MessageDialog.h"
 #include "PropertyCustomizationHelpers.h"
 #include "SBlueprintEditorToolbar.h"
+#include "ScopedTransaction.h"
 #include "SGraphPanel.h"
 #include "SKismetInspector.h"
 #include "SMyBlueprint.h"
@@ -39,6 +40,40 @@ namespace UE::RenderGrid::Private
 }
 
 
+URenderGridJobSelection::URenderGridJobSelection()
+{
+	SetFlags(RF_Public | RF_Transactional);
+}
+
+bool URenderGridJobSelection::SetSelectedRenderGridJobs(const TArray<URenderGridJob*>& Jobs)
+{
+	TSet<FGuid> PreviouslySelectedJobIds;
+	PreviouslySelectedJobIds.Append(SelectedRenderGridJobIds);
+
+	SelectedRenderGridJobIds.Empty();
+	for (URenderGridJob* Job : Jobs)
+	{
+		if (!IsValid(Job))
+		{
+			continue;
+		}
+		SelectedRenderGridJobIds.Add(Job->GetGuid());
+	}
+	if (SelectedRenderGridJobIds.Num() != PreviouslySelectedJobIds.Num())
+	{
+		return true;
+	}
+	for (const FGuid& JobId : SelectedRenderGridJobIds)
+	{
+		if (!PreviouslySelectedJobIds.Contains(JobId))
+		{
+			return true;
+		}
+	}
+	return false;
+}
+
+
 void UE::RenderGrid::Private::FRenderGridEditor::InitRenderGridEditor(const EToolkitMode::Type Mode, const TSharedPtr<IToolkitHost>& InitToolkitHost, URenderGridBlueprint* InRenderGridBlueprint)
 {
 	check(InRenderGridBlueprint);
@@ -87,6 +122,7 @@ void UE::RenderGrid::Private::FRenderGridEditor::InitRenderGridEditor(const EToo
 
 UE::RenderGrid::Private::FRenderGridEditor::FRenderGridEditor()
 	: PreviewBlueprint(nullptr)
+	, RenderGridJobSelection(NewObject<URenderGridJobSelection>())
 	, bRunRenderNewBatch(false)
 	, BatchRenderQueue(nullptr)
 	, PreviewRenderQueue(nullptr)
@@ -182,6 +218,22 @@ void UE::RenderGrid::Private::FRenderGridEditor::RegisterApplicationModes(const
 	}
 }
 
+void UE::RenderGrid::Private::FRenderGridEditor::PostUndo(bool bSuccessful)
+{
+	FBlueprintEditor::PostUndo(bSuccessful);
+
+	OnRenderGridChanged().Broadcast();
+	OnRenderGridJobsSelectionChanged().Broadcast();
+}
+
+void UE::RenderGrid::Private::FRenderGridEditor::PostRedo(bool bSuccessful)
+{
+	FBlueprintEditor::PostRedo(bSuccessful);
+
+	OnRenderGridChanged().Broadcast();
+	OnRenderGridJobsSelectionChanged().Broadcast();
+}
+
 void UE::RenderGrid::Private::FRenderGridEditor::Compile()
 {
 	DestroyInstance();
@@ -214,17 +266,17 @@ void UE::RenderGrid::Private::FRenderGridEditor::SetIsDebugging(const bool bInIs
 
 bool UE::RenderGrid::Private::FRenderGridEditor::IsBatchRendering() const
 {
-	return IsValid(BatchRenderQueue);
+	return IsValid(BatchRenderQueue.Get());
 }
 
 bool UE::RenderGrid::Private::FRenderGridEditor::IsPreviewRendering() const
 {
-	return IsValid(PreviewRenderQueue);
+	return IsValid(PreviewRenderQueue.Get());
 }
 
 void UE::RenderGrid::Private::FRenderGridEditor::SetPreviewRenderQueue(URenderGridQueue* Queue)
 {
-	PreviewRenderQueue = Queue;
+	PreviewRenderQueue = TStrongObjectPtr(Queue);
 	SetIsDebugging(IsValid(Queue));
 }
 
@@ -245,11 +297,14 @@ TArray<URenderGridJob*> UE::RenderGrid::Private::FRenderGridEditor::GetSelectedR
 	TArray<URenderGridJob*> SelectedJobs;
 	if (URenderGrid* Grid = GetInstance(); IsValid(Grid))
 	{
-		for (URenderGridJob* Job : Grid->GetRenderGridJobs())
+		if (URenderGridJobSelection* JobSelection = RenderGridJobSelection.Get(); IsValid(JobSelection))
 		{
-			if (SelectedRenderGridJobIds.Contains(Job->GetGuid()))
+			for (URenderGridJob* Job : Grid->GetRenderGridJobs())
 			{
-				SelectedJobs.Add(Job);
+				if (JobSelection->SelectedRenderGridJobIds.Contains(Job->GetGuid()))
+				{
+					SelectedJobs.Add(Job);
+				}
 			}
 		}
 	}
@@ -258,31 +313,18 @@ TArray<URenderGridJob*> UE::RenderGrid::Private::FRenderGridEditor::GetSelectedR
 
 void UE::RenderGrid::Private::FRenderGridEditor::SetSelectedRenderGridJobs(const TArray<URenderGridJob*>& Jobs)
 {
-	TSet<FGuid> PreviouslySelectedJobIds;
-	PreviouslySelectedJobIds.Append(SelectedRenderGridJobIds);
-
-	SelectedRenderGridJobIds.Empty();
-	for (URenderGridJob* Job : Jobs)
+	if (URenderGridJobSelection* JobSelection = RenderGridJobSelection.Get(); IsValid(JobSelection))
 	{
-		if (!IsValid(Job))
 		{
-			continue;
+			FScopedTransaction Transaction(LOCTEXT("ChangeJobSelection", "Change Job Selection"));
+			JobSelection->Modify();
+			if (!JobSelection->SetSelectedRenderGridJobs(Jobs))
+			{
+				return;
+			}
 		}
-		SelectedRenderGridJobIds.Add(Job->GetGuid());
-	}
 
-	if (SelectedRenderGridJobIds.Num() != PreviouslySelectedJobIds.Num())
-	{
 		OnRenderGridJobsSelectionChanged().Broadcast();
-		return;
-	}
-	for (const FGuid& JobId : SelectedRenderGridJobIds)
-	{
-		if (!PreviouslySelectedJobIds.Contains(JobId))
-		{
-			OnRenderGridJobsSelectionChanged().Broadcast();
-			return;
-		}
 	}
 }
 
@@ -416,10 +458,12 @@ void UE::RenderGrid::Private::FRenderGridEditor::AddJobAction()
 {
 	if (URenderGrid* Instance = GetInstance(); IsValid(Instance))
 	{
+		FScopedTransaction Transaction(LOCTEXT("AddJob", "Add Job"));
+		MarkAsModified();
+
 		if (URenderGridJob* Job = Instance->CreateAndAddNewRenderGridJob(); IsValid(Job))
 		{
 			OnRenderGridJobCreated().Broadcast(Job);
-			MarkAsModified();
 			OnRenderGridChanged().Broadcast();
 		}
 	}
@@ -435,11 +479,13 @@ void UE::RenderGrid::Private::FRenderGridEditor::DuplicateJobAction()
 
 	if (URenderGrid* Instance = GetInstance(); IsValid(Instance))
 	{
+		FScopedTransaction Transaction(LOCTEXT("DuplicateJob", "Duplicate Job"));
+		MarkAsModified();
+
 		for (URenderGridJob* SelectedJob : SelectedJobs)
 		{
 			Instance->DuplicateAndAddRenderGridJob(SelectedJob);
 		}
-		MarkAsModified();
 		OnRenderGridChanged().Broadcast();
 	}
 }
@@ -465,12 +511,13 @@ void UE::RenderGrid::Private::FRenderGridEditor::DeleteJobAction()
 
 	if (URenderGrid* Instance = GetInstance(); IsValid(Instance))
 	{
+		FScopedTransaction Transaction(LOCTEXT("DeleteJob", "Delete Job"));
+		MarkAsModified();
+
 		for (URenderGridJob* SelectedJob : SelectedJobs)
 		{
 			Instance->RemoveRenderGridJob(SelectedJob);
 		}
-
-		MarkAsModified();
 		OnRenderGridChanged().Broadcast();
 	}
 }
@@ -498,7 +545,7 @@ void UE::RenderGrid::Private::FRenderGridEditor::BatchRenderListAction()
 		if (URenderGridQueue* RenderQueue = IRenderGridModule::Get().GetManager().CreateBatchRenderQueue(Instance); IsValid(RenderQueue))
 		{
 			RenderQueue->OnExecuteFinished().AddRaw(this, &FRenderGridEditor::OnBatchRenderListActionFinished);
-			BatchRenderQueue = RenderQueue;
+			BatchRenderQueue = TStrongObjectPtr(RenderQueue);
 
 			SetIsDebugging(true);
 			RenderQueue->Execute();
@@ -512,7 +559,7 @@ void UE::RenderGrid::Private::FRenderGridEditor::OnBatchRenderListActionFinished
 {
 	SetIsDebugging(false);
 
-	URenderGridQueue* FinishedRenderQueue = BatchRenderQueue;
+	URenderGridQueue* FinishedRenderQueue = BatchRenderQueue.Get();
 	BatchRenderQueue = nullptr;
 	OnRenderGridBatchRenderingFinished().Broadcast(FinishedRenderQueue);
 }
@@ -590,6 +637,9 @@ void UE::RenderGrid::Private::FRenderGridEditor::DestroyInstance()
 		// Garbage collection
 		RenderGridWeakPtr.Reset();
 		Instance->MarkAsGarbage();
+
+		// Clear the undo/redo buffer, to prevent any bugs/glitches with undo/redo while the render grid editor is closed
+		GEditor->ResetTransaction(LOCTEXT("RenderGridEditorClosed", "Render Grid Editor Closed"));
 	}
 }
 
@@ -602,7 +652,7 @@ void UE::RenderGrid::Private::FRenderGridEditor::UpdateInstance(UBlueprint* InBl
 		DestroyInstance();
 
 		// Save the Blueprint we're creating a preview for
-		PreviewBlueprint = Cast<URenderGridBlueprint>(InBlueprint);
+		PreviewBlueprint = TStrongObjectPtr(Cast<URenderGridBlueprint>(InBlueprint));
 
 		URenderGrid* RenderGrid;
 		// Create the Widget, we have to do special swapping out of the widget tree.
diff --git a/Engine/Plugins/Experimental/RenderGrid/Source/RenderGridEditor/Private/Toolkit/RenderGridEditor.h b/Engine/Plugins/Experimental/RenderGrid/Source/RenderGridEditor/Private/Toolkit/RenderGridEditor.h
--- a/Engine/Plugins/Experimental/RenderGrid/Source/RenderGridEditor/Private/Toolkit/RenderGridEditor.h
+++ b/Engine/Plugins/Experimental/RenderGrid/Source/RenderGridEditor/Private/Toolkit/RenderGridEditor.h
@@ -4,6 +4,7 @@
 
 #include "CoreMinimal.h"
 #include "IRenderGridEditor.h"
+#include "RenderGridEditor.generated.h"
 
 
 class FSpawnTabArgs;
@@ -20,6 +21,28 @@ namespace UE::RenderGrid::Private
 }
 
 
+/**
+ * The selected render grid jobs.
+ * This is placed in a separate class, so it can be an UObject, so it can be placed in the undo/redo buffer.
+ */
+UCLASS()
+class RENDERGRIDEDITOR_API URenderGridJobSelection : public UObject
+{
+	GENERATED_BODY()
+
+public:
+	URenderGridJobSelection();
+
+	/** Sets the new job selection. Returns true if it is different from the previous selection, returns false if it was exactly the same. */
+	bool SetSelectedRenderGridJobs(const TArray<URenderGridJob*>& Jobs);
+
+public:
+	/** The GUIDs of the currently selected render grid jobs. */
+	UPROPERTY()
+	TSet<FGuid> SelectedRenderGridJobIds;
+};
+
+
 namespace UE::RenderGrid::Private
 {
 	DECLARE_MULTICAST_DELEGATE_TwoParams(FRenderGridEditorClosed, const UE::RenderGrid::IRenderGridEditor*, URenderGridBlueprint*);
@@ -49,7 +72,8 @@ namespace UE::RenderGrid::Private
 		virtual void OnActiveTabChanged(TSharedPtr<SDockTab> PreviouslyActive, TSharedPtr<SDockTab> NewlyActivated) override;
 		virtual void SetupGraphEditorEvents(UEdGraph* InGraph, SGraphEditor::FGraphEditorEvents& InEvents) override;
 		virtual void RegisterApplicationModes(const TArray<UBlueprint*>& InBlueprints, bool bShouldOpenInDefaultsMode, bool bNewlyCreated = false) override;
-
+		virtual void PostUndo(bool bSuccessful) override;
+		virtual void PostRedo(bool bSuccessful) override;
 		virtual void Compile() override;
 		//~ End FBlueprintEditor Interface
 
@@ -59,9 +83,9 @@ namespace UE::RenderGrid::Private
 		virtual void SetIsDebugging(const bool bInIsDebugging) override;
 		virtual TSharedPtr<FRenderGridBlueprintEditorToolbar> GetRenderGridToolbarBuilder() override { return RenderGridToolbar; }
 		virtual bool IsBatchRendering() const override;
-		virtual URenderGridQueue* GetBatchRenderQueue() const override { return BatchRenderQueue; }
+		virtual URenderGridQueue* GetBatchRenderQueue() const override { return BatchRenderQueue.Get(); }
 		virtual bool IsPreviewRendering() const override;
-		virtual URenderGridQueue* GetPreviewRenderQueue() const override { return PreviewRenderQueue; }
+		virtual URenderGridQueue* GetPreviewRenderQueue() const override { return PreviewRenderQueue.Get(); }
 		virtual void SetPreviewRenderQueue(URenderGridQueue* Queue) override;
 		virtual void MarkAsModified() override;
 		virtual TArray<URenderGridJob*> GetSelectedRenderGridJobs() const override;
@@ -157,22 +181,22 @@ namespace UE::RenderGrid::Private
 		TSharedPtr<FExtender> ToolbarExtender;
 
 		/** The blueprint instance that's currently visible in the editor. */
-		TObjectPtr<URenderGridBlueprint> PreviewBlueprint;
+		TStrongObjectPtr<URenderGridBlueprint> PreviewBlueprint;
 
 		/** The current render grid instance that's visible in the editor. */
 		mutable TWeakObjectPtr<URenderGrid> RenderGridWeakPtr;
 
-		/** The IDs of the currently selected render grid jobs. */
-		TSet<FGuid> SelectedRenderGridJobIds;
+		/** The currently selected render grid jobs. */
+		TStrongObjectPtr<URenderGridJobSelection> RenderGridJobSelection;
 
 		/** True if it should call BatchRenderListAction() next frame. */
 		bool bRunRenderNewBatch;
 
 		/** The current batch render queue, if any. */
-		TObjectPtr<URenderGridQueue> BatchRenderQueue;
+		TStrongObjectPtr<URenderGridQueue> BatchRenderQueue;
 
 		/** The current preview render queue, if any. */
-		TObjectPtr<URenderGridQueue> PreviewRenderQueue;
+		TStrongObjectPtr<URenderGridQueue> PreviewRenderQueue;
 
 		/** The time it should still remain in debugging mode (after it has been turned off). */
 		float DebuggingTimeInSecondsRemaining;
diff --git a/Engine/Plugins/Experimental/RenderGrid/Source/RenderGridEditor/Private/UI/SRenderGrid.cpp b/Engine/Plugins/Experimental/RenderGrid/Source/RenderGridEditor/Private/UI/SRenderGrid.cpp
--- a/Engine/Plugins/Experimental/RenderGrid/Source/RenderGridEditor/Private/UI/SRenderGrid.cpp
+++ b/Engine/Plugins/Experimental/RenderGrid/Source/RenderGridEditor/Private/UI/SRenderGrid.cpp
@@ -11,7 +11,7 @@
 
 void UE::RenderGrid::Private::SRenderGrid::Tick(const FGeometry&, const double, const float)
 {
-	if (DetailsView.IsValid())
+	if (RenderGridDetailsView.IsValid())
 	{
 		if (const TSharedPtr<IRenderGridEditor> BlueprintEditor = BlueprintEditorWeakPtr.Pin())
 		{
@@ -21,10 +21,9 @@ void UE::RenderGrid::Private::SRenderGrid::Tick(const FGeometry&, const double,
 				{
 					Grid = nullptr;
 				}
-				if (DetailsViewRenderGridWeakPtr != Grid)
+				if (RenderGridWeakPtr != Grid)
 				{
-					DetailsViewRenderGridWeakPtr = Grid;
-					DetailsView->SetObject(Grid);
+					SetRenderGrid(Grid);
 				}
 			}
 		}
@@ -36,20 +35,36 @@ BEGIN_SLATE_FUNCTION_BUILD_OPTIMIZATION
 void UE::RenderGrid::Private::SRenderGrid::Construct(const FArguments& InArgs, TSharedPtr<IRenderGridEditor> InBlueprintEditor)
 {
 	BlueprintEditorWeakPtr = InBlueprintEditor;
-	DetailsViewRenderGridWeakPtr = InBlueprintEditor->GetInstance();
+	RenderGridWeakPtr = nullptr;
 
 	FPropertyEditorModule& PropertyEditorModule = FModuleManager::GetModuleChecked<FPropertyEditorModule>("PropertyEditor");
+
 	FDetailsViewArgs DetailsViewArgs;
 	DetailsViewArgs.bAllowSearch = false;
 	DetailsViewArgs.NameAreaSettings = FDetailsViewArgs::HideNameArea;
 	DetailsViewArgs.DefaultsOnlyVisibility = EEditDefaultsOnlyNodeVisibility::Hide;
 	DetailsViewArgs.NotifyHook = this;
-	DetailsView = PropertyEditorModule.CreateDetailView(DetailsViewArgs);
-	DetailsView->SetObject(DetailsViewRenderGridWeakPtr.Get());
+
+	RenderGridDetailsView = PropertyEditorModule.CreateDetailView(DetailsViewArgs);
+	RenderGridDefaultsDetailsView = PropertyEditorModule.CreateDetailView(DetailsViewArgs);
+
+	SetRenderGrid(InBlueprintEditor->GetInstance());
 
 	ChildSlot
 	[
-		DetailsView->AsShared()
+		SNew(SVerticalBox)
+
+		+ SVerticalBox::Slot()
+		.AutoHeight()
+		[
+			RenderGridDetailsView->AsShared()
+		]
+
+		+ SVerticalBox::Slot()
+		.AutoHeight()
+		[
+			RenderGridDefaultsDetailsView->AsShared()
+		]
 	];
 }
 
@@ -66,5 +81,12 @@ void UE::RenderGrid::Private::SRenderGrid::NotifyPostChange(const FPropertyChang
 	}
 }
 
+void UE::RenderGrid::Private::SRenderGrid::SetRenderGrid(URenderGrid* RenderGrid)
+{
+	RenderGridWeakPtr = RenderGrid;
+	RenderGridDetailsView->SetObject(RenderGrid);
+	RenderGridDefaultsDetailsView->SetObject(IsValid(RenderGrid) ? RenderGrid->GetDefaultsObject() : nullptr);
+}
+
 
 #undef LOCTEXT_NAMESPACE
diff --git a/Engine/Plugins/Experimental/RenderGrid/Source/RenderGridEditor/Private/UI/SRenderGrid.h b/Engine/Plugins/Experimental/RenderGrid/Source/RenderGridEditor/Private/UI/SRenderGrid.h
--- a/Engine/Plugins/Experimental/RenderGrid/Source/RenderGridEditor/Private/UI/SRenderGrid.h
+++ b/Engine/Plugins/Experimental/RenderGrid/Source/RenderGridEditor/Private/UI/SRenderGrid.h
@@ -8,6 +8,7 @@
 
 
 class URenderGrid;
+class URenderGridDefaults;
 class IDetailsView;
 
 namespace UE::RenderGrid
@@ -36,14 +37,20 @@ namespace UE::RenderGrid::Private
 		virtual void NotifyPostChange(const FPropertyChangedEvent& PropertyChangedEvent, FEditPropertyChain* PropertyThatChanged) override;
 		//~ End FNotifyHook Interface
 
+	protected:
+		void SetRenderGrid(URenderGrid* RenderGrid);
+
 	private:
 		/** A reference to the blueprint editor that owns the render grid instance. */
 		TWeakPtr<IRenderGridEditor> BlueprintEditorWeakPtr;
 
+		/** The render grid instance that's being edited in the details view. */
+		TWeakObjectPtr<URenderGrid> RenderGridWeakPtr;
+
 		/** A reference to the details view. */
-		TSharedPtr<IDetailsView> DetailsView;
+		TSharedPtr<IDetailsView> RenderGridDetailsView;
 
-		/** The render grid instance that's being edited in the details view. */
-		TWeakObjectPtr<URenderGrid> DetailsViewRenderGridWeakPtr;
+		/** A reference to the details view. */
+		TSharedPtr<IDetailsView> RenderGridDefaultsDetailsView;
 	};
 }
diff --git a/Engine/Plugins/Experimental/RenderGrid/Source/RenderGridEditor/Private/UI/SRenderGridJobList.cpp b/Engine/Plugins/Experimental/RenderGrid/Source/RenderGridEditor/Private/UI/SRenderGridJobList.cpp
--- a/Engine/Plugins/Experimental/RenderGrid/Source/RenderGridEditor/Private/UI/SRenderGridJobList.cpp
+++ b/Engine/Plugins/Experimental/RenderGrid/Source/RenderGridEditor/Private/UI/SRenderGridJobList.cpp
@@ -3,23 +3,25 @@
 #include "UI/SRenderGridJobList.h"
 
 #include "Commands/RenderGridEditorCommands.h"
-#include "Framework/Application/SlateApplication.h"
-#include "Layout/WidgetPath.h"
 #include "UI/Components/SRenderGridDragHandle.h"
 #include "UI/Components/SRenderGridEditableTextBlock.h"
 #include "UI/Components/SRenderGridFileSelectorTextBlock.h"
 #include "RenderGrid/RenderGrid.h"
 #include "RenderGrid/RenderGridQueue.h"
 #include "IRenderGridEditor.h"
+#include "RenderGridUtils.h"
 
 #include "Framework/MultiBox/MultiBoxBuilder.h"
 #include "Misc/MessageDialog.h"
 #include "MoviePipelineExecutor.h"
 #include "MoviePipelineOutputSetting.h"
+#include "MoviePipelinePrimaryConfig.h"
 #include "MoviePipelineQueue.h"
 #include "MovieScene.h"
 #include "PropertyCustomizationHelpers.h"
+#include "ScopedTransaction.h"
 #include "SlateOptMacros.h"
+#include "Framework/Application/SlateApplication.h"
 #include "Styling/AppStyle.h"
 #include "Tracks/MovieSceneSubTrack.h"
 #include "Widgets/Input/SCheckBox.h"
@@ -406,7 +408,7 @@ void UE::RenderGrid::Private::SRenderGridJobListTable::HandleJobListSelectionCha
 void UE::RenderGrid::Private::SRenderGridJobListTableRow::Construct(const FArguments& InArgs, const TSharedRef<STableViewBase>& InOwnerTableView, TWeakPtr<IRenderGridEditor> InBlueprintEditor, URenderGridJob* InRenderGridJob, const TSharedPtr<SRenderGridJobList>& InJobListWidget)
 {
 	BlueprintEditorWeakPtr = InBlueprintEditor;
-	RenderGridJob = InRenderGridJob;
+	RenderGridJobWeakPtr = InRenderGridJob;
 	JobListWidget = InJobListWidget;
 
 	SMultiColumnTableRow<URenderGridJob*>::Construct(FSuperRowType::FArguments()
@@ -459,225 +461,272 @@ FReply UE::RenderGrid::Private::SRenderGridJobListTableRow::OnAcceptDrop(const F
 
 TSharedRef<SWidget> UE::RenderGrid::Private::SRenderGridJobListTableRow::GenerateWidgetForColumn(const FName& ColumnName)
 {
-	if (!IsValid(RenderGridJob))
-	{
-		return SNullWidget::NullWidget;
-	}
-
-	if (ColumnName == FRenderGridJobListColumns::DragDropHandle)
+	if (URenderGridJob* Job = RenderGridJobWeakPtr.Get(); IsValid(Job))
 	{
-		return SNew(SBox)
-			.Padding(FMargin(0.0f, 2.0f, 2.0f, 2.0f))
-			[
-				SNew(SScaleBox)
+		if (ColumnName == FRenderGridJobListColumns::DragDropHandle)
+		{
+			return SNew(SBox)
+				.Padding(FMargin(0.0f, 2.0f, 2.0f, 2.0f))
+				[
+					SNew(SScaleBox)
+					.HAlign(HAlign_Center)
+					.VAlign(VAlign_Center)
+					.Stretch(EStretch::ScaleToFit)
+					.StretchDirection(EStretchDirection::Both)
+					[
+						SNew(SRenderGridDragHandle<FRenderGridJobListTableRowDragDropOp>, Job)
+						.Widget(SharedThis(this))
+					]
+				];
+		}
+		else if (ColumnName == FRenderGridJobListColumns::IsEnabled)
+		{
+			return SNew(SBox)
 				.HAlign(HAlign_Center)
-				.VAlign(VAlign_Center)
-				.Stretch(EStretch::ScaleToFit)
-				.StretchDirection(EStretchDirection::Both)
 				[
-					SNew(SRenderGridDragHandle<FRenderGridJobListTableRowDragDropOp>, RenderGridJob)
-					.Widget(SharedThis(this))
-				]
-			];
-	}
-	else if (ColumnName == FRenderGridJobListColumns::IsEnabled)
-	{
-		return SNew(SBox)
-			.HAlign(HAlign_Center)
-			[
-				SNew(SCheckBox)
-				.IsChecked(RenderGridJob->GetIsEnabled())
-				.OnCheckStateChanged_Lambda([this](ECheckBoxState State)
+					SNew(SCheckBox)
+					.IsChecked(Job->GetIsEnabled())
+					.OnCheckStateChanged_Lambda([this](ECheckBoxState State)
+					{
+						if (URenderGridJob* RenderGridJob = RenderGridJobWeakPtr.Get(); IsValid(RenderGridJob))
+						{
+							{
+								FScopedTransaction Transaction(LOCTEXT("ChangeJobId", "Change Job Id"));
+								if (const TSharedPtr<IRenderGridEditor> BlueprintEditor = BlueprintEditorWeakPtr.Pin())
+								{
+									BlueprintEditor->MarkAsModified();
+								}
+								RenderGridJob->Modify();
+
+								RenderGridJob->SetIsEnabled(State == ECheckBoxState::Checked);
+							}
+
+							if (JobListWidget.IsValid())
+							{
+								JobListWidget->RefreshHeaderEnabledCheckbox();
+							}
+						}
+					})
+				];
+		}
+		else if (ColumnName == FRenderGridJobListColumns::JobId)
+		{
+			return SNew(SRenderGridEditableTextBlock)
+				.Text(FText::FromString(Job->GetJobId()))
+				.OnTextCommitted_Lambda([this](const FText& InLabel, ETextCommit::Type InCommitInfo) -> FText
+				{
+					if (URenderGridJob* RenderGridJob = RenderGridJobWeakPtr.Get(); IsValid(RenderGridJob))
+					{
+						const FString OldJobId = RenderGridJob->GetJobId();
+						const FString NewJobId = FRenderGridUtils::PurgeJobIdOrReturnEmptyString(InLabel.ToString());
+						if (NewJobId.IsEmpty() || (RenderGridJob->GetJobId() == NewJobId))
+						{
+							return FText::FromString(OldJobId);
+						}
+
+						if (const TSharedPtr<IRenderGridEditor> BlueprintEditor = BlueprintEditorWeakPtr.Pin())
+						{
+							if (URenderGrid* Grid = BlueprintEditor->GetInstance(); IsValid(Grid))
+							{
+								if (Grid->DoesJobIdExist(NewJobId))
+								{
+									const FText TitleText = LOCTEXT("JobIdNotUniqueTitle", "Duplicate Job IDs");
+									FMessageDialog::Open(
+										EAppMsgType::Ok,
+										FText::Format(LOCTEXT("JobIdNotUniqueMessage", "Job ID \"{0}\" is not unique."), FText::FromString(NewJobId)),
+										&TitleText);
+									return FText::FromString(OldJobId);
+								}
+
+								FScopedTransaction Transaction(LOCTEXT("ChangeJobId", "Change Job Id"));
+								BlueprintEditor->MarkAsModified();
+								RenderGridJob->Modify();
+
+								RenderGridJob->SetJobId(NewJobId);
+								return FText::FromString(RenderGridJob->GetJobId());
+							}
+						}
+						return FText::FromString(OldJobId);
+					}
+					return FText();
+				});
+		}
+		else if (ColumnName == FRenderGridJobListColumns::JobName)
+		{
+			return SNew(SRenderGridEditableTextBlock)
+				.Text(FText::FromString(Job->GetJobName()))
+				.OnTextCommitted_Lambda([this](const FText& InLabel, ETextCommit::Type InCommitInfo) -> FText
 				{
-					RenderGridJob->SetIsEnabled(State == ECheckBoxState::Checked);
-					if (JobListWidget.IsValid())
+					if (URenderGridJob* RenderGridJob = RenderGridJobWeakPtr.Get(); IsValid(RenderGridJob))
 					{
-						JobListWidget->RefreshHeaderEnabledCheckbox();
+						FScopedTransaction Transaction(LOCTEXT("ChangeJobName", "Change Job Name"));
+						if (const TSharedPtr<IRenderGridEditor> BlueprintEditor = BlueprintEditorWeakPtr.Pin())
+						{
+							BlueprintEditor->MarkAsModified();
+						}
+						RenderGridJob->Modify();
+
+						RenderGridJob->SetJobName(InLabel.ToString());
+						return FText::FromString(RenderGridJob->GetJobName());
 					}
-					if (const TSharedPtr<IRenderGridEditor> BlueprintEditor = BlueprintEditorWeakPtr.Pin())
+					return FText();
+				});
+		}
+		else if (ColumnName == FRenderGridJobListColumns::OutputDirectory)
+		{
+			return SNew(SRenderGridFileSelectorTextBlock)
+				.Text(FText::FromString(Job->GetOutputDirectoryForDisplay()))
+				.FolderPath_Lambda([this]() -> FString
+				{
+					if (URenderGridJob* RenderGridJob = RenderGridJobWeakPtr.Get(); IsValid(RenderGridJob))
 					{
-						BlueprintEditor->MarkAsModified();
+						return RenderGridJob->GetOutputDirectory();
 					}
+					return FString();
 				})
-			];
-	}
-	else if (ColumnName == FRenderGridJobListColumns::JobId)
-	{
-		return SNew(SRenderGridEditableTextBlock)
-			.Text(FText::FromString(RenderGridJob->GetJobId()))
-			.OnTextCommitted_Lambda([this](const FText& InLabel, ETextCommit::Type InCommitInfo) -> FText
-			{
-				const FString OldJobId = RenderGridJob->GetJobId();
-				const FString NewJobId = URenderGridJob::PurgeJobIdOrReturnEmptyString(InLabel.ToString());
-				if (NewJobId.IsEmpty() || (RenderGridJob->GetJobId() == NewJobId))
-				{
-					return FText::FromString(OldJobId);
-				}
-
-				if (const TSharedPtr<IRenderGridEditor> BlueprintEditor = BlueprintEditorWeakPtr.Pin())
+				.OnTextCommitted_Lambda([this](const FText& InLabel, ETextCommit::Type InCommitInfo) -> FText
 				{
-					if (URenderGrid* Grid = BlueprintEditor->GetInstance(); IsValid(Grid))
+					if (URenderGridJob* RenderGridJob = RenderGridJobWeakPtr.Get(); IsValid(RenderGridJob))
 					{
-						if (Grid->DoesJobIdExist(NewJobId))
+						FScopedTransaction Transaction(LOCTEXT("ChangeJobOutputDirectory", "Change Job Output Directory"));
+						if (const TSharedPtr<IRenderGridEditor> BlueprintEditor = BlueprintEditorWeakPtr.Pin())
 						{
-							const FText TitleText = LOCTEXT("JobIdNotUniqueTitle", "Duplicate Job IDs");
-							FMessageDialog::Open(
-								EAppMsgType::Ok,
-								FText::Format(LOCTEXT("JobIdNotUniqueMessage", "Job ID \"{0}\" is not unique."), FText::FromString(NewJobId)),
-								&TitleText);
-							return FText::FromString(OldJobId);
+							BlueprintEditor->MarkAsModified();
 						}
+						RenderGridJob->Modify();
 
-						RenderGridJob->SetJobId(NewJobId);
-						BlueprintEditor->MarkAsModified();
-						return FText::FromString(RenderGridJob->GetJobId());
+						RenderGridJob->SetOutputDirectory(InLabel.ToString());
+						return FText::FromString(RenderGridJob->GetOutputDirectoryForDisplay());
 					}
-				}
-				return FText::FromString(OldJobId);
-			});
-	}
-	else if (ColumnName == FRenderGridJobListColumns::JobName)
-	{
-		return SNew(SRenderGridEditableTextBlock)
-			.Text(FText::FromString(RenderGridJob->GetJobName()))
-			.OnTextCommitted_Lambda([this](const FText& InLabel, ETextCommit::Type InCommitInfo) -> FText
-			{
-				RenderGridJob->SetJobName(InLabel.ToString());
-				if (const TSharedPtr<IRenderGridEditor> BlueprintEditor = BlueprintEditorWeakPtr.Pin())
-				{
-					BlueprintEditor->MarkAsModified();
-				}
-				return FText::FromString(RenderGridJob->GetJobName());
-			});
-	}
-	else if (ColumnName == FRenderGridJobListColumns::OutputDirectory)
-	{
-		return SNew(SRenderGridFileSelectorTextBlock)
-			.Text(FText::FromString(RenderGridJob->GetOutputDirectoryForDisplay()))
-			.FolderPath_Lambda([this]() -> FString
-			{
-				return RenderGridJob->GetOutputDirectory();
-			})
-			.OnTextCommitted_Lambda([this](const FText& InLabel, ETextCommit::Type InCommitInfo) -> FText
-			{
-				RenderGridJob->SetOutputDirectory(InLabel.ToString());
-				if (const TSharedPtr<IRenderGridEditor> BlueprintEditor = BlueprintEditorWeakPtr.Pin())
-				{
-					BlueprintEditor->MarkAsModified();
-				}
-				return FText::FromString(RenderGridJob->GetOutputDirectoryForDisplay());
-			});
-	}
-	else if (ColumnName == FRenderGridJobListColumns::RenderPreset)
-	{
-		return SNew(SObjectPropertyEntryBox)
-			.AllowedClass(UMoviePipelinePrimaryConfig::StaticClass())
-			.ObjectPath_Lambda([this]() -> FString
-			{
-				if (UMoviePipelinePrimaryConfig* Preset = RenderGridJob->GetRenderPreset(); IsValid(Preset))
-				{
-					return Preset->GetPathName();
-				}
-				return FString();
-			})
-			.OnObjectChanged_Lambda([this](const FAssetData& AssetData) -> void
-			{
-				RenderGridJob->SetRenderPreset(nullptr);
-				if (UObject* AssetDataAsset = AssetData.GetAsset(); IsValid(AssetDataAsset))
+					return FText();
+				});
+		}
+		else if (ColumnName == FRenderGridJobListColumns::RenderPreset)
+		{
+			return SNew(SObjectPropertyEntryBox)
+				.AllowedClass(UMoviePipelinePrimaryConfig::StaticClass())
+				.ObjectPath_Lambda([this]() -> FString
 				{
-					if (UMoviePipelinePrimaryConfig* Preset = Cast<UMoviePipelinePrimaryConfig>(AssetDataAsset))
+					if (URenderGridJob* RenderGridJob = RenderGridJobWeakPtr.Get(); IsValid(RenderGridJob))
 					{
-						RenderGridJob->SetRenderPreset(Preset);
+						if (UMoviePipelinePrimaryConfig* Preset = RenderGridJob->GetRenderPreset(); IsValid(Preset))
+						{
+							return Preset->GetPathName();
+						}
 					}
-				}
-				if (const TSharedPtr<IRenderGridEditor> BlueprintEditor = BlueprintEditorWeakPtr.Pin())
+					return FString();
+				})
+				.OnObjectChanged_Lambda([this](const FAssetData& AssetData) -> void
 				{
-					BlueprintEditor->MarkAsModified();
-					BlueprintEditor->OnRenderGridChanged().Broadcast();
-				}
-			})
-			.AllowClear(true)
-			.DisplayUseSelected(true)
-			.DisplayBrowse(true)
-			.DisplayThumbnail(false);
-	}
-	else if ((ColumnName == FRenderGridJobListColumns::StartFrame) || (ColumnName == FRenderGridJobListColumns::EndFrame))
-	{
-		FText Text;
-		if (TOptional<int32> Frame = ((ColumnName == FRenderGridJobListColumns::StartFrame) ? RenderGridJob->GetStartFrame() : RenderGridJob->GetEndFrame()))
+					if (URenderGridJob* RenderGridJob = RenderGridJobWeakPtr.Get(); IsValid(RenderGridJob))
+					{
+						{
+							FScopedTransaction Transaction(LOCTEXT("ChangeJobRenderPreset", "Change Job Render Preset"));
+							if (const TSharedPtr<IRenderGridEditor> BlueprintEditor = BlueprintEditorWeakPtr.Pin())
+							{
+								BlueprintEditor->MarkAsModified();
+							}
+							RenderGridJob->Modify();
+
+							RenderGridJob->SetRenderPreset(nullptr);
+							if (UObject* AssetDataAsset = AssetData.GetAsset(); IsValid(AssetDataAsset))
+							{
+								if (UMoviePipelinePrimaryConfig* Preset = Cast<UMoviePipelinePrimaryConfig>(AssetDataAsset))
+								{
+									RenderGridJob->SetRenderPreset(Preset);
+								}
+							}
+						}
+
+						if (const TSharedPtr<IRenderGridEditor> BlueprintEditor = BlueprintEditorWeakPtr.Pin())
+						{
+							BlueprintEditor->OnRenderGridChanged().Broadcast();
+						}
+					}
+				})
+				.AllowClear(true)
+				.DisplayUseSelected(true)
+				.DisplayBrowse(true)
+				.DisplayThumbnail(false);
+		}
+		else if ((ColumnName == FRenderGridJobListColumns::StartFrame) || (ColumnName == FRenderGridJobListColumns::EndFrame))
 		{
-			Text = FText::AsNumber(*Frame);
+			FText Text;
+			if (TOptional<int32> Frame = ((ColumnName == FRenderGridJobListColumns::StartFrame) ? Job->GetStartFrame() : Job->GetEndFrame()))
+			{
+				Text = FText::AsNumber(*Frame);
+			}
+			return SNew(SBox)
+				.VAlign(VAlign_Center)
+				//.HAlign(HAlign_Right)
+				.Padding(FMargin(10, 0))
+				[
+					SNew(STextBlock).Text(Text)
+				];
 		}
-		return SNew(SBox)
-			.VAlign(VAlign_Center)
-			//.HAlign(HAlign_Right)
-			.Padding(FMargin(10, 0))
-			[
-				SNew(STextBlock).Text(Text)
-			];
-	}
-	else if (ColumnName == FRenderGridJobListColumns::Tags)
-	{
-		//TODO: add support for tags
-	}
-	else if (ColumnName == FRenderGridJobListColumns::Resolution)
-	{
-		FIntPoint Resolution = RenderGridJob->GetOutputResolution();
+		else if (ColumnName == FRenderGridJobListColumns::Tags)
+		{
+			//TODO: add support for tags
+		}
+		else if (ColumnName == FRenderGridJobListColumns::Resolution)
+		{
+			FIntPoint Resolution = Job->GetOutputResolution();
 
-		FNumberFormattingOptions NumberFormattingOptions;
-		NumberFormattingOptions.UseGrouping = false;
+			FNumberFormattingOptions NumberFormattingOptions;
+			NumberFormattingOptions.UseGrouping = false;
 
-		FText ResolutionFormat = LOCTEXT("Resolution", "{Width} x {Height}");
-		FFormatNamedArguments ResolutionArguments;
-		ResolutionArguments.Add(TEXT("Width"), FText::AsNumber(Resolution.X, &NumberFormattingOptions));
-		ResolutionArguments.Add(TEXT("Height"), FText::AsNumber(Resolution.Y, &NumberFormattingOptions));
-		FText Text = FText::Format(ResolutionFormat, ResolutionArguments);
+			FText ResolutionFormat = LOCTEXT("Resolution", "{Width} x {Height}");
+			FFormatNamedArguments ResolutionArguments;
+			ResolutionArguments.Add(TEXT("Width"), FText::AsNumber(Resolution.X, &NumberFormattingOptions));
+			ResolutionArguments.Add(TEXT("Height"), FText::AsNumber(Resolution.Y, &NumberFormattingOptions));
+			FText Text = FText::Format(ResolutionFormat, ResolutionArguments);
 
-		return SNew(SBox)
-			.VAlign(VAlign_Center)
-			.Padding(FMargin(10, 0))
-			[
-				SNew(STextBlock).Text(Text)
-			];
-	}
-	else if (ColumnName == FRenderGridJobListColumns::Duration)
-	{
-		FText Text;
-		if (TOptional<double> Duration = RenderGridJob->GetDurationInSeconds())
+			return SNew(SBox)
+				.VAlign(VAlign_Center)
+				.Padding(FMargin(10, 0))
+				[
+					SNew(STextBlock).Text(Text)
+				];
+		}
+		else if (ColumnName == FRenderGridJobListColumns::Duration)
 		{
-			FTimespan Timespan = FTimespan::FromSeconds(*Duration);
-			int32 Hours = static_cast<int32>(Timespan.GetTotalHours());
-			int32 Minutes = Timespan.GetMinutes();
-			int32 Seconds = Timespan.GetSeconds();
-
-			FNumberFormattingOptions NumberFormattingOptions;
-			NumberFormattingOptions.MinimumIntegralDigits = 2;
-			NumberFormattingOptions.MaximumIntegralDigits = 2;
-
-			FText TimespanFormat = NSLOCTEXT("Timespan", "Format_HoursMinutesSeconds", "{Hours}:{Minutes}:{Seconds}");
-			FFormatNamedArguments TimeArguments;
-			TimeArguments.Add(TEXT("Hours"), Hours);
-			TimeArguments.Add(TEXT("Minutes"), FText::AsNumber(Minutes, &NumberFormattingOptions));
-			TimeArguments.Add(TEXT("Seconds"), FText::AsNumber(Seconds, &NumberFormattingOptions));
-			Text = FText::Format(TimespanFormat, TimeArguments);
+			FText Text;
+			if (TOptional<double> Duration = Job->GetDurationInSeconds())
+			{
+				FTimespan Timespan = FTimespan::FromSeconds(*Duration);
+				int32 Hours = static_cast<int32>(Timespan.GetTotalHours());
+				int32 Minutes = Timespan.GetMinutes();
+				int32 Seconds = Timespan.GetSeconds();
+
+				FNumberFormattingOptions NumberFormattingOptions;
+				NumberFormattingOptions.MinimumIntegralDigits = 2;
+				NumberFormattingOptions.MaximumIntegralDigits = 2;
+
+				FText TimespanFormat = NSLOCTEXT("Timespan", "Format_HoursMinutesSeconds", "{Hours}:{Minutes}:{Seconds}");
+				FFormatNamedArguments TimeArguments;
+				TimeArguments.Add(TEXT("Hours"), Hours);
+				TimeArguments.Add(TEXT("Minutes"), FText::AsNumber(Minutes, &NumberFormattingOptions));
+				TimeArguments.Add(TEXT("Seconds"), FText::AsNumber(Seconds, &NumberFormattingOptions));
+				Text = FText::Format(TimespanFormat, TimeArguments);
+			}
+			return SNew(SBox)
+				.VAlign(VAlign_Center)
+				[
+					SNew(STextBlock).Text(Text)
+				];
+		}
+		else if (ColumnName == FRenderGridJobListColumns::RenderingStatus)
+		{
+			return SNew(SBox)
+				.VAlign(VAlign_Center)
+				[
+					SNew(STextBlock)
+					.Text_Lambda([this]() -> FText
+					{
+						return GetRenderStatusText();
+					})
+				];
 		}
-		return SNew(SBox)
-			.VAlign(VAlign_Center)
-			[
-				SNew(STextBlock).Text(Text)
-			];
-	}
-	else if (ColumnName == FRenderGridJobListColumns::RenderingStatus)
-	{
-		return SNew(SBox)
-			.VAlign(VAlign_Center)
-			[
-				SNew(STextBlock)
-				.Text_Lambda([this]() -> FText
-				{
-					return GetRenderStatusText();
-				})
-			];
 	}
 	return SNullWidget::NullWidget;
 }
@@ -688,7 +737,10 @@ FText UE::RenderGrid::Private::SRenderGridJobListTableRow::GetRenderStatusText()
 	{
 		if (URenderGridQueue* RenderQueue = BlueprintEditor->GetBatchRenderQueue(); IsValid(RenderQueue))
 		{
-			return FText::FromString(RenderQueue->GetJobStatus(RenderGridJob));
+			if (URenderGridJob* RenderGridJob = RenderGridJobWeakPtr.Get(); IsValid(RenderGridJob))
+			{
+				return FText::FromString(RenderQueue->GetJobStatus(RenderGridJob));
+			}
 		}
 	}
 	return FText();
diff --git a/Engine/Plugins/Experimental/RenderGrid/Source/RenderGridEditor/Private/UI/SRenderGridJobList.h b/Engine/Plugins/Experimental/RenderGrid/Source/RenderGridEditor/Private/UI/SRenderGridJobList.h
--- a/Engine/Plugins/Experimental/RenderGrid/Source/RenderGridEditor/Private/UI/SRenderGridJobList.h
+++ b/Engine/Plugins/Experimental/RenderGrid/Source/RenderGridEditor/Private/UI/SRenderGridJobList.h
@@ -158,7 +158,7 @@ namespace UE::RenderGrid::Private
 		TWeakPtr<IRenderGridEditor> BlueprintEditorWeakPtr;
 
 		/** A reference to the render grid job instance. */
-		TObjectPtr<URenderGridJob> RenderGridJob;
+		TWeakObjectPtr<URenderGridJob> RenderGridJobWeakPtr;
 
 		/** A reference to the job list (the parent widget). */
 		TSharedPtr<SRenderGridJobList> JobListWidget;
diff --git a/Engine/Plugins/Experimental/RenderGrid/Source/RenderGridEditor/Private/UI/SRenderGridProps.cpp b/Engine/Plugins/Experimental/RenderGrid/Source/RenderGridEditor/Private/UI/SRenderGridProps.cpp
--- a/Engine/Plugins/Experimental/RenderGrid/Source/RenderGridEditor/Private/UI/SRenderGridProps.cpp
+++ b/Engine/Plugins/Experimental/RenderGrid/Source/RenderGridEditor/Private/UI/SRenderGridProps.cpp
@@ -36,7 +36,8 @@ void UE::RenderGrid::Private::SRenderGridProps::Construct(const FArguments& InAr
 		.BorderImage(new FSlateNoResource());
 
 	Refresh();
-	//InBlueprintEditor->OnRenderGridChanged().AddSP(this, &SRenderGridProps::Refresh);
+	//InBlueprintEditor->OnRenderGridChanged().AddSP(this, &SRenderGridProps::Refresh); // should technically be correct, but causes the remote control properties to bug out
+	InBlueprintEditor->OnRenderGridPreviewRenderedFirstTimeSinceAppStart().AddSP(this, &SRenderGridProps::Refresh);
 	InBlueprintEditor->OnRenderGridJobsSelectionChanged().AddSP(this, &SRenderGridProps::Refresh);
 	InBlueprintEditor->OnRenderGridBatchRenderingStarted().AddSP(this, &SRenderGridProps::OnBatchRenderingStarted);
 	InBlueprintEditor->OnRenderGridBatchRenderingFinished().AddSP(this, &SRenderGridProps::OnBatchRenderingFinished);
diff --git a/Engine/Plugins/Experimental/RenderGrid/Source/RenderGridEditor/Private/UI/SRenderGridPropsLocal.cpp b/Engine/Plugins/Experimental/RenderGrid/Source/RenderGridEditor/Private/UI/SRenderGridPropsLocal.cpp
--- a/Engine/Plugins/Experimental/RenderGrid/Source/RenderGridEditor/Private/UI/SRenderGridPropsLocal.cpp
+++ b/Engine/Plugins/Experimental/RenderGrid/Source/RenderGridEditor/Private/UI/SRenderGridPropsLocal.cpp
@@ -13,7 +13,7 @@ BEGIN_SLATE_FUNCTION_BUILD_OPTIMIZATION
 void UE::RenderGrid::Private::SRenderGridPropsLocal::Construct(const FArguments& InArgs, TSharedPtr<IRenderGridEditor> InBlueprintEditor, URenderGridPropsSourceLocal* InPropsSource)
 {
 	BlueprintEditorWeakPtr = InBlueprintEditor;
-	PropsSource = InPropsSource;
+	PropsSourceWeakPtr = InPropsSource;
 }
 
 END_SLATE_FUNCTION_BUILD_OPTIMIZATION
diff --git a/Engine/Plugins/Experimental/RenderGrid/Source/RenderGridEditor/Private/UI/SRenderGridPropsLocal.h b/Engine/Plugins/Experimental/RenderGrid/Source/RenderGridEditor/Private/UI/SRenderGridPropsLocal.h
--- a/Engine/Plugins/Experimental/RenderGrid/Source/RenderGridEditor/Private/UI/SRenderGridPropsLocal.h
+++ b/Engine/Plugins/Experimental/RenderGrid/Source/RenderGridEditor/Private/UI/SRenderGridPropsLocal.h
@@ -31,6 +31,6 @@ namespace UE::RenderGrid::Private
 		TWeakPtr<IRenderGridEditor> BlueprintEditorWeakPtr;
 
 		/** The props source control. */
-		TObjectPtr<URenderGridPropsSourceLocal> PropsSource;
+		TWeakObjectPtr<URenderGridPropsSourceLocal> PropsSourceWeakPtr;
 	};
 }
diff --git a/Engine/Plugins/Experimental/RenderGrid/Source/RenderGridEditor/Private/UI/SRenderGridPropsRemoteControl.cpp b/Engine/Plugins/Experimental/RenderGrid/Source/RenderGridEditor/Private/UI/SRenderGridPropsRemoteControl.cpp
--- a/Engine/Plugins/Experimental/RenderGrid/Source/RenderGridEditor/Private/UI/SRenderGridPropsRemoteControl.cpp
+++ b/Engine/Plugins/Experimental/RenderGrid/Source/RenderGridEditor/Private/UI/SRenderGridPropsRemoteControl.cpp
@@ -5,6 +5,7 @@
 #include "UI/Components/SRenderGridRemoteControlField.h"
 #include "IRenderGridEditor.h"
 #include "IRenderGridModule.h"
+#include "ScopedTransaction.h"
 #include "RenderGrid/RenderGrid.h"
 #include "RenderGrid/RenderGridManager.h"
 #include "SlateOptMacros.h"
@@ -17,15 +18,15 @@ BEGIN_SLATE_FUNCTION_BUILD_OPTIMIZATION
 void UE::RenderGrid::Private::SRenderGridPropsRemoteControl::Construct(const FArguments& InArgs, TSharedPtr<IRenderGridEditor> InBlueprintEditor, URenderGridPropsSourceRemoteControl* InPropsSource)
 {
 	BlueprintEditorWeakPtr = InBlueprintEditor;
-	PropsSource = InPropsSource;
+	PropsSourceWeakPtr = InPropsSource;
 
 	SAssignNew(RowWidgetsContainer, SVerticalBox);
 	UpdateStoredValuesAndRefresh(true);
 
 	InBlueprintEditor->OnRenderGridJobsSelectionChanged().AddSP(this, &SRenderGridPropsRemoteControl::OnRenderGridJobsSelectionChanged);
-	if (IsValid(PropsSource))
+	if (IsValid(InPropsSource))
 	{
-		if (TObjectPtr<URemoteControlPreset> Preset = PropsSource->GetProps()->GetRemoteControlPreset())
+		if (TObjectPtr<URemoteControlPreset> Preset = InPropsSource->GetProps()->GetRemoteControlPreset())
 		{
 			Preset->OnEntityExposed().AddSP(this, &SRenderGridPropsRemoteControl::OnRemoteControlEntitiesExposed);
 			Preset->OnEntityUnexposed().AddSP(this, &SRenderGridPropsRemoteControl::OnRemoteControlEntitiesUnexposed);
@@ -68,7 +69,7 @@ void UE::RenderGrid::Private::SRenderGridPropsRemoteControl::Refresh(const bool
 	if (const TSharedPtr<IRenderGridEditor> BlueprintEditor = BlueprintEditorWeakPtr.Pin())
 	{
 		TArray<FRenderGridRemoteControlGenerateWidgetArgs> NewRowWidgetsArgs;
-		if (IsValid(PropsSource))
+		if (URenderGridPropsSourceRemoteControl* PropsSource = PropsSourceWeakPtr.Get(); IsValid(PropsSource))
 		{
 			URenderGridPropsRemoteControl* Props = PropsSource->GetProps();
 			for (URenderGridPropRemoteControl* Prop : Props->GetAllCasted())
@@ -218,6 +219,8 @@ bool UE::RenderGrid::Private::SRenderGridPropsRemoteControl::SetSelectedJobField
 {
 	if (URenderGridJob* SelectedJob = GetSelectedJob(); IsValid(SelectedJob))
 	{
+		FScopedTransaction Transaction(LOCTEXT("ChangeJobProperty", "Change Job Property"));
+		SelectedJob->Modify();
 		return SelectedJob->SetRemoteControlValue(RemoteControlEntity, BinaryArray);
 	}
 	return false;
diff --git a/Engine/Plugins/Experimental/RenderGrid/Source/RenderGridEditor/Private/UI/SRenderGridPropsRemoteControl.h b/Engine/Plugins/Experimental/RenderGrid/Source/RenderGridEditor/Private/UI/SRenderGridPropsRemoteControl.h
--- a/Engine/Plugins/Experimental/RenderGrid/Source/RenderGridEditor/Private/UI/SRenderGridPropsRemoteControl.h
+++ b/Engine/Plugins/Experimental/RenderGrid/Source/RenderGridEditor/Private/UI/SRenderGridPropsRemoteControl.h
@@ -59,7 +59,7 @@ namespace UE::RenderGrid::Private
 		TWeakPtr<IRenderGridEditor> BlueprintEditorWeakPtr;
 
 		/** The props source control. */
-		TObjectPtr<URenderGridPropsSourceRemoteControl> PropsSource;
+		TWeakObjectPtr<URenderGridPropsSourceRemoteControl> PropsSourceWeakPtr;
 
 		/** The widget that lists the property rows. */
 		TSharedPtr<SVerticalBox> RowWidgetsContainer;
diff --git a/Engine/Plugins/Experimental/RenderGrid/Source/RenderGridEditor/Private/UI/SRenderGridViewer.cpp b/Engine/Plugins/Experimental/RenderGrid/Source/RenderGridEditor/Private/UI/SRenderGridViewer.cpp
--- a/Engine/Plugins/Experimental/RenderGrid/Source/RenderGridEditor/Private/UI/SRenderGridViewer.cpp
+++ b/Engine/Plugins/Experimental/RenderGrid/Source/RenderGridEditor/Private/UI/SRenderGridViewer.cpp
@@ -108,17 +108,23 @@ void UE::RenderGrid::Private::SRenderGridViewer::Refresh()
 		CachedViewerMode = CurrentViewerMode;
 		WidgetContainer->ClearContent();
 
-		if (CurrentViewerMode == ERenderGridViewerMode::Live)
+		switch (CachedViewerMode)
 		{
-			WidgetContainer->SetContent(SNew(SRenderGridViewerLive, BlueprintEditor));
-		}
-		else if (CurrentViewerMode == ERenderGridViewerMode::Preview)
-		{
-			WidgetContainer->SetContent(SNew(SRenderGridViewerPreview, BlueprintEditor));
-		}
-		else if (CurrentViewerMode == ERenderGridViewerMode::Rendered)
-		{
-			WidgetContainer->SetContent(SNew(SRenderGridViewerRendered, BlueprintEditor));
+			case ERenderGridViewerMode::Preview:
+				WidgetContainer->SetContent(SNew(SRenderGridViewerPreview, BlueprintEditor));
+				break;
+
+			case ERenderGridViewerMode::Rendered:
+				WidgetContainer->SetContent(SNew(SRenderGridViewerRendered, BlueprintEditor));
+				break;
+
+			case ERenderGridViewerMode::None:
+				break;
+
+			case ERenderGridViewerMode::Live:
+			default:
+				WidgetContainer->SetContent(SNew(SRenderGridViewerLive, BlueprintEditor));
+				break;
 		}
 	}
 }
diff --git a/Engine/Plugins/Experimental/RenderGrid/Source/RenderGridEditor/Private/UI/SRenderGridViewerLive.cpp b/Engine/Plugins/Experimental/RenderGrid/Source/RenderGridEditor/Private/UI/SRenderGridViewerLive.cpp
--- a/Engine/Plugins/Experimental/RenderGrid/Source/RenderGridEditor/Private/UI/SRenderGridViewerLive.cpp
+++ b/Engine/Plugins/Experimental/RenderGrid/Source/RenderGridEditor/Private/UI/SRenderGridViewerLive.cpp
@@ -299,7 +299,7 @@ void UE::RenderGrid::Private::SRenderGridViewerLive::Construct(const FArguments&
 					{
 						if (URenderGridJob* SelectedJob = SelectedJobWeakPtr.Get(); IsValid(SelectedJob))
 						{
-							if (ULevelSequence* Sequence = SelectedJob->GetSequence(); IsValid(Sequence))
+							if (ULevelSequence* Sequence = SelectedJob->GetLevelSequence(); IsValid(Sequence))
 							{
 								return LOCTEXT("WaitingForRenderer", "Waiting for renderer...");
 							}
@@ -387,7 +387,7 @@ void UE::RenderGrid::Private::SRenderGridViewerLive::UpdateViewport()
 
 	if (URenderGridJob* SelectedJob = SelectedJobWeakPtr.Get(); IsValid(SelectedJob))
 	{
-		if (ULevelSequence* Sequence = SelectedJob->GetSequence(); IsValid(Sequence))
+		if (ULevelSequence* Sequence = SelectedJob->GetLevelSequence(); IsValid(Sequence))
 		{
 			if (const TSharedPtr<IRenderGridEditor> BlueprintEditor = BlueprintEditorWeakPtr.Pin())
 			{
diff --git a/Engine/Plugins/Experimental/RenderGrid/Source/RenderGridEditor/Private/UI/SRenderGridViewerPreview.cpp b/Engine/Plugins/Experimental/RenderGrid/Source/RenderGridEditor/Private/UI/SRenderGridViewerPreview.cpp
--- a/Engine/Plugins/Experimental/RenderGrid/Source/RenderGridEditor/Private/UI/SRenderGridViewerPreview.cpp
+++ b/Engine/Plugins/Experimental/RenderGrid/Source/RenderGridEditor/Private/UI/SRenderGridViewerPreview.cpp
@@ -1,10 +1,9 @@
 // Copyright Epic Games, Inc. All Rights Reserved.
 
 #include "UI/SRenderGridViewerPreview.h"
-#include "Engine/Texture2D.h"
-#include "TextureResource.h"
 #include "UI/Components/SRenderGridViewerFrameSlider.h"
 #include "RenderGrid/RenderGrid.h"
+#include "RenderGrid/RenderGridQueue.h"
 #include "IRenderGridEditor.h"
 #include "IRenderGridModule.h"
 #include "LevelSequence.h"
@@ -12,12 +11,13 @@
 #include "MovieScene.h"
 #include "RenderGrid/RenderGridManager.h"
 #include "SlateOptMacros.h"
+#include "TextureResource.h"
+#include "Engine/Texture2D.h"
 #include "Sections/MovieSceneSubSection.h"
 #include "Styles/RenderGridEditorStyle.h"
 #include "Tracks/MovieSceneSubTrack.h"
 #include "Widgets/Colors/SColorPicker.h"
 #include "Widgets/Input/SButton.h"
-#include "Widgets/Input/SSlider.h"
 #include "Widgets/Layout/SScaleBox.h"
 
 #define LOCTEXT_NAMESPACE "SRenderGridViewerPreview"
@@ -30,7 +30,7 @@ void UE::RenderGrid::Private::SRenderGridViewerPreview::Tick(const FGeometry&, c
 {
 	if (const TSharedPtr<IRenderGridEditor> BlueprintEditor = BlueprintEditorWeakPtr.Pin())
 	{
-		if (BlueprintEditor->CanCurrentlyRender() && !CurrentRenderQueue)
+		if (BlueprintEditor->CanCurrentlyRender() && !CurrentRenderQueueWeakPtr.IsValid())
 		{
 			UpdateRerenderButton();
 			UpdateFrameSlider();
@@ -53,7 +53,7 @@ void UE::RenderGrid::Private::SRenderGridViewerPreview::Construct(const FArgumen
 {
 	BlueprintEditorWeakPtr = InBlueprintEditor;
 	SelectedJobWeakPtr = nullptr;
-	CurrentRenderQueue = nullptr;
+	CurrentRenderQueueWeakPtr = nullptr;
 	FramesUntilRenderNewPreview = 0;
 	ImageBrushEmpty = FSlateBrush();
 	ImageBrushEmpty.DrawAs = ESlateBrushDrawType::Type::NoDrawType;
@@ -81,7 +81,7 @@ void UE::RenderGrid::Private::SRenderGridViewerPreview::Construct(const FArgumen
 	ChildSlot
 	[
 		SNew(SVerticalBox)
-		.Visibility_Lambda([this]() -> EVisibility { return (!IsPreviewWidget() && CurrentRenderQueue) ? EVisibility::Hidden : EVisibility::Visible; })
+		.Visibility_Lambda([this]() -> EVisibility { return (!IsPreviewWidget() && CurrentRenderQueueWeakPtr.IsValid()) ? EVisibility::Hidden : EVisibility::Visible; })
 
 		// image
 		+ SVerticalBox::Slot()
@@ -123,7 +123,7 @@ void UE::RenderGrid::Private::SRenderGridViewerPreview::Construct(const FArgumen
 						{
 							if (URenderGridJob* SelectedJob = SelectedJobWeakPtr.Get(); IsValid(SelectedJob))
 							{
-								if (ULevelSequence* Sequence = SelectedJob->GetSequence(); IsValid(Sequence))
+								if (ULevelSequence* Sequence = SelectedJob->GetLevelSequence(); IsValid(Sequence))
 								{
 									if (IsPreviewWidget() || !FrameSlider.IsValid() || (FrameSlider->GetVisibility() != EVisibility::Visible))
 									{
@@ -227,9 +227,9 @@ void UE::RenderGrid::Private::SRenderGridViewerPreview::OnObjectModified(UObject
 			GridDataChanged();
 		}
 	}
-	else if (SelectedJobWeakPtr.IsValid() && SelectedJobWeakPtr->GetSequence() && (Cast<UMovieSceneSequence>(Object) || Cast<UMovieScene>(Object) || Cast<UMovieSceneTrack>(Object) || Cast<UMovieSceneSection>(Object) || Cast<UMovieSceneSubTrack>(Object) || Cast<UMovieSceneSubSection>(Object)))
+	else if (SelectedJobWeakPtr.IsValid() && SelectedJobWeakPtr->GetLevelSequence() && (Cast<UMovieSceneSequence>(Object) || Cast<UMovieScene>(Object) || Cast<UMovieSceneTrack>(Object) || Cast<UMovieSceneSection>(Object) || Cast<UMovieSceneSubTrack>(Object) || Cast<UMovieSceneSubSection>(Object)))
 	{
-		ULevelSequence* Sequence = SelectedJobWeakPtr->GetSequence();
+		ULevelSequence* Sequence = SelectedJobWeakPtr->GetLevelSequence();
 		if ((Object == Sequence) || (Object == Sequence->GetMovieScene()) || (Object->GetTypedOuter<UMovieScene>() == Sequence->GetMovieScene()))
 		{
 			// level sequence changed
@@ -338,8 +338,8 @@ void UE::RenderGrid::Private::SRenderGridViewerPreview::InternalRenderNewPreview
 			JobArgs.bHeadless = bHasRenderedSinceAppStart;
 			JobArgs.Frame = SelectedFrame.Get(0);
 		}
-		JobArgs.RenderGrid = BlueprintEditor->GetInstance();
-		JobArgs.RenderGridJob = Job;
+		JobArgs.RenderGrid = TStrongObjectPtr(BlueprintEditor->GetInstance());
+		JobArgs.RenderGridJob = TStrongObjectPtr(Job);
 		JobArgs.Resolution = FIntPoint(FMath::RoundToInt32(RenderResolution), FMath::RoundToInt32(RenderResolution / Job->GetOutputAspectRatio()));
 
 		const TSharedPtr<SWidget> BaseThis = AsShared();
@@ -358,8 +358,8 @@ void UE::RenderGrid::Private::SRenderGridViewerPreview::InternalRenderNewPreview
 
 		if (URenderGridQueue* NewRenderQueue = IRenderGridModule::Get().GetManager().RenderPreviewFrame(JobArgs))
 		{
-			CurrentRenderQueue = NewRenderQueue;
-			BlueprintEditor->SetPreviewRenderQueue(CurrentRenderQueue);
+			CurrentRenderQueueWeakPtr = NewRenderQueue;
+			BlueprintEditor->SetPreviewRenderQueue(NewRenderQueue);
 			return;
 		}
 	}
@@ -368,14 +368,18 @@ void UE::RenderGrid::Private::SRenderGridViewerPreview::InternalRenderNewPreview
 
 void UE::RenderGrid::Private::SRenderGridViewerPreview::RenderNewPreviewCallback(const bool bSuccess)
 {
-	bHasRenderedSinceAppStart = true;
 	UpdateImageTexture();
 
-	CurrentRenderQueue = nullptr;
+	CurrentRenderQueueWeakPtr = nullptr;
 	if (const TSharedPtr<IRenderGridEditor> BlueprintEditor = BlueprintEditorWeakPtr.Pin())
 	{
-		BlueprintEditor->SetPreviewRenderQueue(CurrentRenderQueue);
+		BlueprintEditor->SetPreviewRenderQueue(nullptr);
+		if (!bHasRenderedSinceAppStart)
+		{
+			BlueprintEditor->OnRenderGridPreviewRenderedFirstTimeSinceAppStart().Broadcast();
+		}
 	}
+	bHasRenderedSinceAppStart = true;
 }
 
 void UE::RenderGrid::Private::SRenderGridViewerPreview::UpdateImageTexture(const bool bForce)
@@ -440,20 +444,13 @@ void UE::RenderGrid::Private::SRenderGridViewerPreview::SetImageTexture(UTexture
 		}
 		ImageBrush.SetResourceObject(nullptr);
 		ImageBrush.SetImageSize(FVector2D(0, 0));
-
-		if (IsValid(ImageTexture))
-		{
-			ImageTexture->RemoveFromRoot();
-		}
 		ImageTexture = nullptr;
 	}// cleanup <<
 
 	{// set new texture >>
 		if (IsValid(Texture) && Texture->GetResource())
 		{
-			ImageTexture = Texture;
-			ImageTexture->AddToRoot();
-
+			ImageTexture = TStrongObjectPtr(Texture);
 			if (Image.IsValid())
 			{
 				static constexpr double PreviewTabAspectRatio = 1.96875;
@@ -475,7 +472,7 @@ void UE::RenderGrid::Private::SRenderGridViewerPreview::SetImageTexture(UTexture
 				ImageBrush = FSlateBrush();
 				ImageBrush.DrawAs = ESlateBrushDrawType::Type::Image;
 				ImageBrush.ImageType = ESlateBrushImageType::Type::FullColor;
-				ImageBrush.SetResourceObject(ImageTexture);
+				ImageBrush.SetResourceObject(ImageTexture.Get());
 				ImageBrush.SetImageSize(FVector2D(ImageWidth, ImageHeight));
 				Image->SetImage(&ImageBrush);
 				ImageBackground->SetImage(FCoreStyle::Get().GetBrush("Checkerboard"));
diff --git a/Engine/Plugins/Experimental/RenderGrid/Source/RenderGridEditor/Private/UI/SRenderGridViewerPreview.h b/Engine/Plugins/Experimental/RenderGrid/Source/RenderGridEditor/Private/UI/SRenderGridViewerPreview.h
--- a/Engine/Plugins/Experimental/RenderGrid/Source/RenderGridEditor/Private/UI/SRenderGridViewerPreview.h
+++ b/Engine/Plugins/Experimental/RenderGrid/Source/RenderGridEditor/Private/UI/SRenderGridViewerPreview.h
@@ -3,8 +3,10 @@
 #pragma once
 
 #include "CoreMinimal.h"
-#include "Styling/SlateBrush.h"
+#include "Engine/Texture2D.h"
+#include "UObject/StrongObjectPtr.h"
 #include "Widgets/SCompoundWidget.h"
+#include "Styling/SlateBrush.h"
 
 
 class SButton;
@@ -12,7 +14,6 @@ class SImage;
 class SSlider;
 class URenderGridJob;
 class URenderGridQueue;
-class UTexture2D;
 
 
 namespace UE::RenderGrid
@@ -75,7 +76,7 @@ namespace UE::RenderGrid::Private
 		TWeakPtr<IRenderGridEditor> BlueprintEditorWeakPtr;
 
 		/** A reference to the queue that's currently rendering. */
-		TObjectPtr<URenderGridQueue> CurrentRenderQueue;
+		TWeakObjectPtr<URenderGridQueue> CurrentRenderQueueWeakPtr;
 
 		/** A reference to the job that's currently rendering. */
 		TWeakObjectPtr<URenderGridJob> SelectedJobWeakPtr;
@@ -106,7 +107,7 @@ namespace UE::RenderGrid::Private
 		FSlateBrush ImageBrush;
 
 		/** The texture of the image. */
-		TObjectPtr<UTexture2D> ImageTexture;
+		TStrongObjectPtr<UTexture2D> ImageTexture;
 
 		/** The last render grid job used for the LastUpdateImageTexture function. */
 		TWeakObjectPtr<URenderGridJob> LastUpdateImageTextureSelectedJobWeakPtr;
diff --git a/Engine/Plugins/Experimental/RenderGrid/Source/RenderGridEditor/Public/IRenderGridEditor.h b/Engine/Plugins/Experimental/RenderGrid/Source/RenderGridEditor/Public/IRenderGridEditor.h
--- a/Engine/Plugins/Experimental/RenderGrid/Source/RenderGridEditor/Public/IRenderGridEditor.h
+++ b/Engine/Plugins/Experimental/RenderGrid/Source/RenderGridEditor/Public/IRenderGridEditor.h
@@ -82,6 +82,9 @@ namespace UE::RenderGrid
 		DECLARE_MULTICAST_DELEGATE_OneParam(FOnRenderGridBatchRenderingFinished, URenderGridQueue*);
 		virtual FOnRenderGridBatchRenderingFinished& OnRenderGridBatchRenderingFinished() { return OnRenderGridBatchRenderingFinishedDelegate; }
 
+		DECLARE_MULTICAST_DELEGATE(FOnRenderGridPreviewRenderedFirstTimeSinceAppStart)
+		virtual FOnRenderGridPreviewRenderedFirstTimeSinceAppStart& OnRenderGridPreviewRenderedFirstTimeSinceAppStart() { return OnRenderGridPreviewRenderedFirstTimeSinceAppStartDelegate; }
+
 	private:
 		/** The delegate for when data in the render grid changed. */
 		FOnRenderGridChanged OnRenderGridChangedDelegate;
@@ -97,5 +100,8 @@ namespace UE::RenderGrid
 
 		/** The delegate for when batch rendering of a render grid has finished, successfully or not. */
 		FOnRenderGridBatchRenderingFinished OnRenderGridBatchRenderingFinishedDelegate;
+
+		/** The delegate for the first time (since the start of this application) that the preview is rendered. */
+		FOnRenderGridPreviewRenderedFirstTimeSinceAppStart OnRenderGridPreviewRenderedFirstTimeSinceAppStartDelegate;
 	};
 }
diff --git a/Engine/Plugins/Experimental/RenderGrid/Source/RenderGridEditor/RenderGridEditor.Build.cs b/Engine/Plugins/Experimental/RenderGrid/Source/RenderGridEditor/RenderGridEditor.Build.cs
--- a/Engine/Plugins/Experimental/RenderGrid/Source/RenderGridEditor/RenderGridEditor.Build.cs
+++ b/Engine/Plugins/Experimental/RenderGrid/Source/RenderGridEditor/RenderGridEditor.Build.cs
@@ -7,6 +7,8 @@ public class RenderGridEditor : ModuleRules
 	public RenderGridEditor(ReadOnlyTargetRules Target) : base(Target)
 	{
 		PCHUsage = PCHUsageMode.UseExplicitOrSharedPCHs;
+		//PCHUsage = PCHUsageMode.NoPCHs;
+		//bUseUnity = false;
 
 		PublicDependencyModuleNames.AddRange(
 			new string[]
@@ -18,7 +20,6 @@ public class RenderGridEditor : ModuleRules
 			}
 		);
 
-
 		PrivateDependencyModuleNames.AddRange(
 			new string[]
 			{
