RenderGrid plugin:
- Fixed an issue where RenderGrid Queue's weren't blueprint types (couldn't place them in a variable etc).

diff --git a/Engine/Plugins/Experimental/RenderGrid/Source/RenderGrid/Public/RenderGrid/RenderGridQueue.h b/Engine/Plugins/Experimental/RenderGrid/Source/RenderGrid/Public/RenderGrid/RenderGridQueue.h
--- a/Engine/Plugins/Experimental/RenderGrid/Source/RenderGrid/Public/RenderGrid/RenderGridQueue.h
+++ b/Engine/Plugins/Experimental/RenderGrid/Source/RenderGrid/Public/RenderGrid/RenderGridQueue.h
@@ -153,7 +153,7 @@ protected:
 /**
  * This class is responsible for rendering the given render grid jobs.
  */
-UCLASS()
+UCLASS(BlueprintType)
 class RENDERGRID_API URenderGridQueue : public UObject
 {
 	GENERATED_BODY()
