- Changed the way rendergrid assets are being saved and load, in order to make the undo/redo buffer not have to be cleared in between saves and recompiles.
- Made it also add job relocation events to the undo/redo buffer.
- Fixed issues (caused by the save-load rework) that prevented rendergrid instances from being deleted.
- Fixed issues (caused by the save-load rework) where marking a variable as "Edit on Instance" would cause a crash.
- Cleaned up the blueprint class (a lot could be removed, since variables will now always have to be forced to turn "Transient" on and some other toggles off, and so the logic required to detect new variables, changed variables, unchanged variables, etc is no longer needed).

diff --git a/Engine/Plugins/Experimental/RenderGrid/Source/RenderGrid/Private/RenderGrid/RenderGrid.cpp b/Engine/Plugins/Experimental/RenderGrid/Source/RenderGrid/Private/RenderGrid/RenderGrid.cpp
--- a/Engine/Plugins/Experimental/RenderGrid/Source/RenderGrid/Private/RenderGrid/RenderGrid.cpp
+++ b/Engine/Plugins/Experimental/RenderGrid/Source/RenderGrid/Private/RenderGrid/RenderGrid.cpp
@@ -433,20 +433,19 @@ bool URenderGridJob::SetRemoteControlValue(const TSharedPtr<FRemoteControlEntity
 }
 
 
-URenderGrid::URenderGrid(const FObjectInitializer& ObjectInitializer): Super(ObjectInitializer)
-	, Guid(FGuid::NewGuid())
+URenderGrid::URenderGrid()
+	: Guid(FGuid::NewGuid())
 	, PropsSourceType(ERenderGridPropsSourceType::RemoteControl)
 	, PropsSourceOrigin_RemoteControl(nullptr)
-	, Defaults(ObjectInitializer.CreateDefaultSubobject<URenderGridDefaults>(this, TEXT("RenderGridDefaults")))
+	, Defaults(CreateDefaultSubobject<URenderGridDefaults>(TEXT("RenderGridDefaults")))
 	, bExecutingBlueprintEvent(false)
 	, CachedPropsSource(nullptr)
 	, CachedPropsSourceType(ERenderGridPropsSourceType::Local)
 {
-	if (!HasAnyFlags(RF_ClassDefaultObject))
+	SetFlags(RF_Public);
+	if (!HasAnyFlags(RF_ClassDefaultObject | RF_DefaultSubObject))
 	{
-		SetFlags(RF_Public | RF_Transactional);
-		LoadValuesFromCDO();
-		OnPreSaveCDO().AddUObject(this, &URenderGrid::SaveValuesToCDO);
+		ClearFlags(RF_Transactional);
 	}
 }
 
@@ -491,25 +490,23 @@ UWorld* URenderGrid::GetWorld() const
 
 void URenderGrid::PreSave(FObjectPreSaveContext SaveContext)
 {
-	if (HasAnyFlags(RF_ClassDefaultObject))
-	{
-		OnPreSaveCDO().Broadcast();
-	}
 	Super::PreSave(SaveContext);
-	SaveValuesToCDO();
 }
 
 void URenderGrid::PostLoad()
 {
 	Super::PostLoad();
-	LoadValuesFromCDO();
 
 	if (PropsSourceType == ERenderGridPropsSourceType::Local)
 	{
 		PropsSourceType = ERenderGridPropsSourceType::RemoteControl;
 	}
 
-	SetFlags(RF_Public | RF_Transactional);
+	SetFlags(RF_Public);
+	if (!HasAnyFlags(RF_ClassDefaultObject | RF_DefaultSubObject))
+	{
+		ClearFlags(RF_Transactional);
+	}
 	for (TObjectPtr<URenderGridJob> Job : RenderGridJobs)
 	{
 		if (IsValid(Job))
@@ -524,6 +521,40 @@ void URenderGrid::PostLoad()
 	Defaults->SetFlags(RF_Public | RF_Transactional);
 }
 
+void URenderGrid::CopyJobs(URenderGrid* From)
+{
+	if (!IsValid(From))
+	{
+		return;
+	}
+	SetRenderGridJobs(From->GetRenderGridJobs());
+}
+
+void URenderGrid::CopyAllPropertiesExceptJobs(URenderGrid* From)
+{
+	if (!IsValid(From))
+	{
+		return;
+	}
+	for (FProperty* Property = StaticClass()->PropertyLink; Property; Property = Property->PropertyLinkNext)
+	{
+		if (!Property->HasAnyPropertyFlags(CPF_Transient | CPF_DuplicateTransient) && Property->GetName().Compare("RenderGridJobs"))// if not [Transient] and not "RenderGridJobs"
+		{
+			Property->CopyCompleteValue(Property->ContainerPtrToValuePtr<void>(this), Property->ContainerPtrToValuePtr<void>(From));
+		}
+	}
+}
+
+void URenderGrid::CopyAllProperties(URenderGrid* From)
+{
+	if (!IsValid(From))
+	{
+		return;
+	}
+	CopyJobs(From);
+	CopyAllPropertiesExceptJobs(From);
+}
+
 void URenderGrid::BeginEditor()
 {
 	bExecutingBlueprintEvent = true;
@@ -587,55 +618,6 @@ void URenderGrid::EndViewportRender(URenderGridJob* Job)
 	bExecutingBlueprintEvent = false;
 }
 
-void URenderGrid::CopyValuesToOrFromCDO(const bool bToCDO)
-{
-	if (HasAnyFlags(RF_ClassDefaultObject))
-	{
-		return;
-	}
-
-	URenderGrid* CDO = GetCDO();
-	if (!IsValid(CDO))
-	{
-		return;
-	}
-
-	for (FProperty* Property = GetClass()->PropertyLink; Property; Property = Property->PropertyLinkNext)
-	{
-		if (!Property->HasAnyPropertyFlags(CPF_Transient | CPF_DuplicateTransient))// if not [Transient]
-		{
-			void* Data = Property->ContainerPtrToValuePtr<void>(this);
-			void* DataCDO = Property->ContainerPtrToValuePtr<void>(CDO);
-			if (bToCDO)
-			{
-				Property->CopyCompleteValue(DataCDO, Data);
-			}
-			else
-			{
-				Property->CopyCompleteValue(Data, DataCDO);
-			}
-		}
-	}
-
-	if (bToCDO)
-	{
-		CDO->RenderGridJobs = TArray<TObjectPtr<URenderGridJob>>();
-		for (URenderGridJob* Job : RenderGridJobs)
-		{
-			if (!IsValid(Job))
-			{
-				continue;
-			}
-			if (URenderGridJob* JobCopy = DuplicateObject(Job, CDO); IsValid(JobCopy))
-			{
-				CDO->RenderGridJobs.Add(JobCopy);
-			}
-		}
-
-		CDO->Defaults = DuplicateObject(Defaults, CDO);
-	}
-}
-
 void URenderGrid::SetPropsSource(ERenderGridPropsSourceType InPropsSourceType, UObject* InPropsSourceOrigin)
 {
 	if (InPropsSourceType == ERenderGridPropsSourceType::RemoteControl)
@@ -702,6 +684,20 @@ void URenderGrid::SetDefaultOutputDirectory(const FString& NewOutputDirectory)
 	Defaults->OutputDirectory = UE::RenderGrid::Private::FRenderGridUtils::NormalizeJobOutputDirectory(NewOutputDirectory);
 }
 
+void URenderGrid::ClearRenderGridJobs()
+{
+	RenderGridJobs.Empty();
+}
+
+void URenderGrid::SetRenderGridJobs(const TArray<URenderGridJob*>& Jobs)
+{
+	RenderGridJobs.Empty();
+	for (URenderGridJob* Job : Jobs)
+	{
+		AddRenderGridJob(Job);
+	}
+}
+
 void URenderGrid::AddRenderGridJob(URenderGridJob* Job)
 {
 	if (IsValid(Job))
diff --git a/Engine/Plugins/Experimental/RenderGrid/Source/RenderGrid/Public/RenderGrid/RenderGrid.h b/Engine/Plugins/Experimental/RenderGrid/Source/RenderGrid/Public/RenderGrid/RenderGrid.h
--- a/Engine/Plugins/Experimental/RenderGrid/Source/RenderGrid/Public/RenderGrid/RenderGrid.h
+++ b/Engine/Plugins/Experimental/RenderGrid/Source/RenderGrid/Public/RenderGrid/RenderGrid.h
@@ -306,7 +306,7 @@ class RENDERGRID_API URenderGrid : public UObject
 	GENERATED_BODY()
 
 public:
-	URenderGrid(const FObjectInitializer& ObjectInitializer);
+	URenderGrid();
 
 	//UObject interface
 	virtual UWorld* GetWorld() const override;
@@ -314,8 +314,14 @@ public:
 	virtual void PostLoad() override;
 	//~ End UObject Interface
 
-	/** Should be called when the editor closes this asset. */
-	void OnClose() { SaveValuesToCDO(); }
+	/** Copy jobs from the given RenderGrid to self. */
+	void CopyJobs(URenderGrid* From);
+
+	/** Copy all properties except jobs from the given RenderGrid to self. */
+	void CopyAllPropertiesExceptJobs(URenderGrid* From);
+
+	/** Copy all properties from the given RenderGrid to self. */
+	void CopyAllProperties(URenderGrid* From);
 
 public:
 	static TArray<FString> GetBlueprintImplementableEvents()
@@ -438,25 +444,6 @@ public:
 	/** Overridable native event for when job rendering for the viewport viewer-mode ends. */
 	virtual void EndViewportRender(URenderGridJob* Job);
 
-private:
-	/**
-	 * Because render grid assets are blueprints (assets that also have a blueprint graph), the render grid data is not stored directly in the asset data that you see in the content browser.
-	 * Instead, the data that is stored (and load) is the CDO (class default object).
-	 * Because of that, any data that needs to persist needs to be copied over to the CDO during a save, and data you'd like to load from it needs to be copied from the CDO during a load.
-	 */
-
-	/** Obtains the CDO, could return itself if this is called on the CDO instance. */
-	URenderGrid* GetCDO() { return (HasAnyFlags(RF_ClassDefaultObject) ? this : GetClass()->GetDefaultObject<URenderGrid>()); }
-
-	/** Copied values over into the CDO. */
-	void SaveValuesToCDO() { CopyValuesToOrFromCDO(true); }
-
-	/** Copied values over from the CDO. */
-	void LoadValuesFromCDO() { CopyValuesToOrFromCDO(false); }
-
-	/** Copied values to or from the CDO, based on whether bToCDO is true or false. */
-	void CopyValuesToOrFromCDO(const bool bToCDO);
-
 public:
 	/** Returns the GUID, which is randomly generated at creation. */
 	UFUNCTION(BlueprintPure, Category="Render Grid")
@@ -515,6 +502,12 @@ public:
 	void SetDefaultOutputDirectoryRaw(const FString& NewOutputDirectory) { Defaults->OutputDirectory = NewOutputDirectory; }
 
 public:
+	UFUNCTION(BlueprintCallable, Category="Render Grid", Meta=(Keywords="remove delete all"))
+	void ClearRenderGridJobs();
+	
+	UFUNCTION(BlueprintCallable, Category="Render Grid")
+	void SetRenderGridJobs(const TArray<URenderGridJob*>& Jobs);
+
 	UFUNCTION(BlueprintCallable, Category="Render Grid")
 	void AddRenderGridJob(URenderGridJob* Job);
 
@@ -583,14 +576,6 @@ public:
 	/** Returns true when it's currently executing a blueprint implementable event, returns false otherwise. */
 	bool IsCurrentlyExecutingUserCode() const { return bExecutingBlueprintEvent; }
 
-private:
-	DECLARE_MULTICAST_DELEGATE(FOnRenderGridPreSave);
-	FOnRenderGridPreSave& OnPreSaveCDO() { return GetCDO()->OnRenderGridPreSaveDelegate; }
-
-private:
-	/** The delegate for when this grid is about to save. */
-	FOnRenderGridPreSave OnRenderGridPreSaveDelegate;
-
 private:
 	/** The unique ID of this render grid. */
 	UPROPERTY()
@@ -602,7 +587,7 @@ private:
 	ERenderGridPropsSourceType PropsSourceType;
 
 	/** The remote control properties that a job in this grid can have, only use this if PropsSourceType is ERenderGridPropsSourceType::RemoteControl. */
-	UPROPERTY(EditInstanceOnly, Category="Render Grid", Meta=(DisplayName="Remote Control Preset", AllowPrivateAccess="true", EditCondition="PropsSourceType == ERenderGridPropsSourceType::RemoteControl", EditConditionHides))
+	UPROPERTY(EditInstanceOnly, Category="Render Grid", Meta=(DisplayName="Remote Control Preset", AllowPrivateAccess="true" /*, EditCondition="PropsSourceType == ERenderGridPropsSourceType::RemoteControl", EditConditionHides*/))
 	TObjectPtr<URemoteControlPreset> PropsSourceOrigin_RemoteControl;
 
 
diff --git a/Engine/Plugins/Experimental/RenderGrid/Source/RenderGridDeveloper/Private/Blueprints/RenderGridBlueprint.cpp b/Engine/Plugins/Experimental/RenderGrid/Source/RenderGridDeveloper/Private/Blueprints/RenderGridBlueprint.cpp
--- a/Engine/Plugins/Experimental/RenderGrid/Source/RenderGridDeveloper/Private/Blueprints/RenderGridBlueprint.cpp
+++ b/Engine/Plugins/Experimental/RenderGrid/Source/RenderGridDeveloper/Private/Blueprints/RenderGridBlueprint.cpp
@@ -9,6 +9,12 @@
 #include "RenderGrid/RenderGridBlueprintGeneratedClass.h"
 
 
+URenderGridBlueprint::URenderGridBlueprint()
+{
+	RenderGrid = CreateDefaultSubobject<URenderGrid>("RenderGrid");
+	RenderGrid->SetFlags(RF_DefaultSubObject | RF_Transactional);
+}
+
 UClass* URenderGridBlueprint::GetBlueprintClass() const
 {
 	return URenderGridBlueprintGeneratedClass::StaticClass();
@@ -18,6 +24,8 @@ void URenderGridBlueprint::PostLoad()
 {
 	Super::PostLoad();
 
+	PropagateAllPropertiesToInstances();
+
 	if (UbergraphPages.IsEmpty() || ((UbergraphPages.Num() == 1) && UbergraphPages[0]->Nodes.IsEmpty()))
 	{
 		if (!UbergraphPages.IsEmpty())
@@ -48,113 +56,123 @@ void URenderGridBlueprint::PostLoad()
 		LastEditedDocuments.AddUnique(NewGraph);
 	}
 
-	FCoreUObjectDelegates::OnObjectModified.RemoveAll(this);
 	OnChanged().RemoveAll(this);
-	FCoreUObjectDelegates::OnObjectModified.AddUObject(this, &URenderGridBlueprint::OnPreVariablesChange);
 	OnChanged().AddUObject(this, &URenderGridBlueprint::OnPostVariablesChange);
 
-	OnPreVariablesChange(this);
 	OnPostVariablesChange(this);
 }
 
-void URenderGridBlueprint::OnPreVariablesChange(UObject* InObject)
+void URenderGridBlueprint::RunOnInstances(const UE::RenderGrid::FRenderGridBlueprintRunOnInstancesCallback& Callback)
 {
-	if (InObject != this)
+	if (UClass* MyRenderGridClass = GeneratedClass; IsValid(MyRenderGridClass))
 	{
-		return;
+		if (URenderGrid* DefaultObject = Cast<URenderGrid>(MyRenderGridClass->GetDefaultObject(false)); IsValid(DefaultObject))
+		{
+			Callback.ExecuteIfBound(DefaultObject);
+
+			TArray<UObject*> ArchetypeInstances;
+			DefaultObject->GetArchetypeInstances(ArchetypeInstances);
+			for (UObject* ArchetypeInstance : ArchetypeInstances)
+			{
+				if (URenderGrid* InstanceGrid = Cast<URenderGrid>(ArchetypeInstance); IsValid(InstanceGrid))
+				{
+					Callback.ExecuteIfBound(InstanceGrid);
+				}
+			}
+		}
 	}
-	LastNewVariables = NewVariables;
 }
 
-void URenderGridBlueprint::OnPostVariablesChange(UBlueprint* InBlueprint)
+void URenderGridBlueprint::Load()
 {
-	if (InBlueprint != this)
+	for (URenderGridJob* Job : RenderGrid->GetRenderGridJobsRef())
 	{
-		return;
-	}
-
-	bool bFoundChange = false;
-
-	TMap<FGuid, int32> NewVariablesByGuid;
-	for (int32 VarIndex = 0; VarIndex < NewVariables.Num(); VarIndex++)
-	{
-		NewVariablesByGuid.Add(NewVariables[VarIndex].VarGuid, VarIndex);
-	}
-
-	TMap<FGuid, int32> OldVariablesByGuid;
-	for (int32 VarIndex = 0; VarIndex < LastNewVariables.Num(); VarIndex++)
-	{
-		OldVariablesByGuid.Add(LastNewVariables[VarIndex].VarGuid, VarIndex);
+		Job->Rename(nullptr, RenderGrid);
 	}
+}
 
-	for (FBPVariableDescription& OldVariable : LastNewVariables)
+void URenderGridBlueprint::Save()
+{
+	if (UClass* MyRenderGridClass = GeneratedClass; IsValid(MyRenderGridClass))
 	{
-		if (!NewVariablesByGuid.Contains(OldVariable.VarGuid))
+		if (URenderGrid* DefaultObject = Cast<URenderGrid>(MyRenderGridClass->GetDefaultObject(false)); IsValid(DefaultObject))
 		{
-			bFoundChange = true;
-			OnVariableRemoved(OldVariable);
+			for (URenderGridJob* Job : DefaultObject->GetRenderGridJobsRef())
+			{
+				Job->Rename(nullptr, DefaultObject);
+			}
 		}
 	}
+}
 
-	for (FBPVariableDescription& NewVariable : NewVariables)
+void URenderGridBlueprint::PropagateJobsToInstances()
+{
+	RunOnInstances(UE::RenderGrid::FRenderGridBlueprintRunOnInstancesCallback::CreateLambda([this](URenderGrid* Instance)
 	{
-		if (!OldVariablesByGuid.Contains(NewVariable.VarGuid))
-		{
-			bFoundChange = true;
-			OnVariableAdded(NewVariable);
-			continue;
-		}
-
-		int32 OldVarIndex = OldVariablesByGuid.FindChecked(NewVariable.VarGuid);
-		const FBPVariableDescription& OldVariable = LastNewVariables[OldVarIndex];
-		if (OldVariable.VarName != NewVariable.VarName)
-		{
-			bFoundChange = true;
-			OnVariableRenamed(NewVariable, OldVariable.VarName, NewVariable.VarName);
-		}
+		Instance->CopyJobs(RenderGrid);
+	}));
+}
 
-		if (OldVariable.VarType != NewVariable.VarType)
-		{
-			bFoundChange = true;
-			OnVariableTypeChanged(NewVariable, OldVariable.VarType, NewVariable.VarType);
-		}
+void URenderGridBlueprint::PropagateAllPropertiesExceptJobsToInstances()
+{
+	RunOnInstances(UE::RenderGrid::FRenderGridBlueprintRunOnInstancesCallback::CreateLambda([this](URenderGrid* Instance)
+	{
+		Instance->CopyAllPropertiesExceptJobs(RenderGrid);
+	}));
+}
 
-		if (OldVariable.PropertyFlags != NewVariable.PropertyFlags)
-		{
-			bFoundChange = true;
-			OnVariablePropertyFlagsChanged(NewVariable, OldVariable.PropertyFlags, NewVariable.PropertyFlags);
-		}
-	}
+void URenderGridBlueprint::PropagateAllPropertiesToInstances()
+{
+	RunOnInstances(UE::RenderGrid::FRenderGridBlueprintRunOnInstancesCallback::CreateLambda([this](URenderGrid* Instance)
+	{
+		Instance->CopyAllProperties(RenderGrid);
+	}));
+}
 
-	LastNewVariables = NewVariables;
+void URenderGridBlueprint::PropagateJobsToAsset(URenderGrid* Instance)
+{
+	check(IsValid(Instance));
+	RenderGrid->CopyJobs(Instance);
+}
 
-	if (bFoundChange)
-	{
-		FBlueprintEditorUtils::MarkBlueprintAsStructurallyModified(this);
-	}
+void URenderGridBlueprint::PropagateAllPropertiesExceptJobsToAsset(URenderGrid* Instance)
+{
+	check(IsValid(Instance));
+	RenderGrid->CopyAllPropertiesExceptJobs(Instance);
 }
 
-void URenderGridBlueprint::OnVariableAdded(FBPVariableDescription& InVar)
+void URenderGridBlueprint::PropagateAllPropertiesToAsset(URenderGrid* Instance)
 {
-	MakeVariableTransientUnlessInstanceEditable(InVar);
+	check(IsValid(Instance));
+	RenderGrid->CopyAllProperties(Instance);
 }
 
-void URenderGridBlueprint::OnVariablePropertyFlagsChanged(FBPVariableDescription& InVar, const uint64 InOldVarPropertyFlags, const uint64 InNewVarPropertyFlags)
+void URenderGridBlueprint::OnPostVariablesChange(UBlueprint* InBlueprint)
 {
-	if ((InOldVarPropertyFlags & CPF_DisableEditOnInstance) != (InNewVarPropertyFlags & CPF_DisableEditOnInstance))// if the value of [Instance Editable] changed
+	if (InBlueprint != this)
 	{
-		MakeVariableTransientUnlessInstanceEditable(InVar);
+		return;
 	}
-}
 
-void URenderGridBlueprint::MakeVariableTransientUnlessInstanceEditable(FBPVariableDescription& InVar)
-{
-	if ((InVar.PropertyFlags & CPF_DisableEditOnInstance) == 0)// if [Instance Editable]
+	bool bFoundChange = false;
+	for (FBPVariableDescription& NewVariable : NewVariables)
 	{
-		InVar.PropertyFlags &= ~CPF_Transient;// set [Transient] to false
+		uint64 PreviousPropertyFlags = NewVariable.PropertyFlags;
+		bFoundChange = bFoundChange || NewVariable.HasMetaData(FBlueprintMetadata::MD_ExposeOnSpawn);
+
+		// fix that's required for the way of saving/loading the render grid,
+		//  the render grid instance (in the blueprint class) isn't of the generated class (but instead it's of the URenderGrid class itself), and so variables can't be saved or load,
+		//  meaning that every variable has to be transient no matter what, otherwise saving or loading the render grid asset will cause a crash
+		NewVariable.PropertyFlags |= CPF_DisableEditOnInstance;
+		NewVariable.PropertyFlags &= ~CPF_ExposeOnSpawn;
+		NewVariable.PropertyFlags |= CPF_Transient;
+		NewVariable.RemoveMetaData(FBlueprintMetadata::MD_ExposeOnSpawn);
+
+		bFoundChange = bFoundChange || (NewVariable.PropertyFlags != PreviousPropertyFlags);
 	}
-	else
+
+	if (bFoundChange)
 	{
-		InVar.PropertyFlags |= CPF_Transient;// set [Transient] to true
+		FBlueprintEditorUtils::MarkBlueprintAsStructurallyModified(this);
 	}
 }
diff --git a/Engine/Plugins/Experimental/RenderGrid/Source/RenderGridDeveloper/Public/Blueprints/RenderGridBlueprint.h b/Engine/Plugins/Experimental/RenderGrid/Source/RenderGridDeveloper/Public/Blueprints/RenderGridBlueprint.h
--- a/Engine/Plugins/Experimental/RenderGrid/Source/RenderGridDeveloper/Public/Blueprints/RenderGridBlueprint.h
+++ b/Engine/Plugins/Experimental/RenderGrid/Source/RenderGridDeveloper/Public/Blueprints/RenderGridBlueprint.h
@@ -4,20 +4,30 @@
 
 #include "CoreMinimal.h"
 #include "Engine/Blueprint.h"
+#include "RenderGrid/RenderGrid.h"
 #include "RenderGridBlueprint.generated.h"
 
 
+namespace UE::RenderGrid
+{
+	/** A delegate for URenderGridBlueprint::PropagateToInstances. */
+	DECLARE_DELEGATE_OneParam(FRenderGridBlueprintRunOnInstancesCallback, URenderGrid* /*Instance*/);
+}
+
+
 /**
  * A UBlueprint child class for the RenderGrid modules.
  *
  * Required in order for a RenderGrid to be able to have a blueprint graph.
  */
-UCLASS(BlueprintType, Meta = (IgnoreClassThumbnail))
+UCLASS(BlueprintType, Meta=(IgnoreClassThumbnail))
 class RENDERGRIDDEVELOPER_API URenderGridBlueprint : public UBlueprint
 {
 	GENERATED_BODY()
 
 public:
+	URenderGridBlueprint();
+
 	//~ Begin UBlueprint Interface
 	virtual bool SupportedByDefaultBlueprintFactory() const override { return false; }
 	virtual bool IsValidForBytecodeOnlyRecompile() const override { return false; }
@@ -34,19 +44,29 @@ public:
 	//~ End UBlueprint Interface
 
 private:
-	virtual void OnPreVariablesChange(UObject* InObject);
-	virtual void OnPostVariablesChange(UBlueprint* InBlueprint);
+	void RunOnInstances(const UE::RenderGrid::FRenderGridBlueprintRunOnInstancesCallback& Callback);
+
+public:
+	void Load();
+	void Save();
+
+	void PropagateJobsToInstances();
+	void PropagateAllPropertiesExceptJobsToInstances();
+	void PropagateAllPropertiesToInstances();
+
+	void PropagateJobsToAsset(URenderGrid* Instance);
+	void PropagateAllPropertiesExceptJobsToAsset(URenderGrid* Instance);
+	void PropagateAllPropertiesToAsset(URenderGrid* Instance);
 
-protected:
-	virtual void OnVariableAdded(FBPVariableDescription& InVar);
-	virtual void OnVariableRemoved(FBPVariableDescription& InVar) {}
-	virtual void OnVariableRenamed(FBPVariableDescription& InVar, const FName& InOldVarName, const FName& InNewVarName) {}
-	virtual void OnVariableTypeChanged(FBPVariableDescription& InVar, const FEdGraphPinType& InOldVarType, const FEdGraphPinType& InNewVarType) {}
-	virtual void OnVariablePropertyFlagsChanged(FBPVariableDescription& InVar, const uint64 InOldVarPropertyFlags, const uint64 InNewVarPropertyFlags);
+public:
+	/** Returns the RenderGrid reference that this RenderGrid asset contains. */
+	UFUNCTION(BlueprintPure, Category="Render Grid")
+	URenderGrid* GetRenderGrid() const { return RenderGrid; }
 
-protected:
-	virtual void MakeVariableTransientUnlessInstanceEditable(FBPVariableDescription& InVar);
+private:
+	virtual void OnPostVariablesChange(UBlueprint* InBlueprint);
 
 private:
-	TArray<FBPVariableDescription> LastNewVariables;
+	UPROPERTY()
+	TObjectPtr<URenderGrid> RenderGrid;
 };
diff --git a/Engine/Plugins/Experimental/RenderGrid/Source/RenderGridEditor/Private/Toolkit/RenderGridEditor.cpp b/Engine/Plugins/Experimental/RenderGrid/Source/RenderGridEditor/Private/Toolkit/RenderGridEditor.cpp
--- a/Engine/Plugins/Experimental/RenderGrid/Source/RenderGridEditor/Private/Toolkit/RenderGridEditor.cpp
+++ b/Engine/Plugins/Experimental/RenderGrid/Source/RenderGridEditor/Private/Toolkit/RenderGridEditor.cpp
@@ -112,7 +112,10 @@ void UE::RenderGrid::Private::FRenderGridEditor::InitRenderGridEditor(const EToo
 	RegenerateMenusAndToolbars();
 
 	UpdateInstance(GetRenderGridBlueprint(), true);
-
+	if (URenderGridBlueprint* RenderGridBlueprint = GetRenderGridBlueprint())
+	{
+		RenderGridBlueprint->PropagateAllPropertiesToInstances();
+	}
 	constexpr bool bShouldOpenInDefaultsMode = true;
 	RegisterApplicationModes(RenderGridBlueprints, bShouldOpenInDefaultsMode, InRenderGridBlueprint->bIsNewlyCreated);
 
@@ -121,7 +124,7 @@ void UE::RenderGrid::Private::FRenderGridEditor::InitRenderGridEditor(const EToo
 }
 
 UE::RenderGrid::Private::FRenderGridEditor::FRenderGridEditor()
-	: PreviewBlueprint(nullptr)
+	: PreviewBlueprintWeakPtr(nullptr)
 	, RenderGridJobSelection(NewObject<URenderGridJobSelection>())
 	, bRunRenderNewBatch(false)
 	, BatchRenderQueue(nullptr)
@@ -222,6 +225,11 @@ void UE::RenderGrid::Private::FRenderGridEditor::PostUndo(bool bSuccessful)
 {
 	FBlueprintEditor::PostUndo(bSuccessful);
 
+	if (URenderGridBlueprint* PreviewBlueprint = PreviewBlueprintWeakPtr.Get(); IsValid(PreviewBlueprint))
+	{
+		PreviewBlueprint->PropagateAllPropertiesToInstances();
+	}
+
 	OnRenderGridChanged().Broadcast();
 	OnRenderGridJobsSelectionChanged().Broadcast();
 }
@@ -230,6 +238,11 @@ void UE::RenderGrid::Private::FRenderGridEditor::PostRedo(bool bSuccessful)
 {
 	FBlueprintEditor::PostRedo(bSuccessful);
 
+	if (URenderGridBlueprint* PreviewBlueprint = PreviewBlueprintWeakPtr.Get(); IsValid(PreviewBlueprint))
+	{
+		PreviewBlueprint->PropagateAllPropertiesToInstances();
+	}
+
 	OnRenderGridChanged().Broadcast();
 	OnRenderGridJobsSelectionChanged().Broadcast();
 }
@@ -282,13 +295,10 @@ void UE::RenderGrid::Private::FRenderGridEditor::SetPreviewRenderQueue(URenderGr
 
 void UE::RenderGrid::Private::FRenderGridEditor::MarkAsModified()
 {
-	if (URenderGrid* Instance = GetInstance(); IsValid(Instance))
-	{
-		Instance->Modify();
-	}
-	if (UBlueprint* BlueprintObj = GetBlueprintObj(); IsValid(BlueprintObj))
+	if (URenderGridBlueprint* BlueprintObj = Cast<URenderGridBlueprint>(GetBlueprintObj()); IsValid(BlueprintObj))
 	{
 		BlueprintObj->Modify();
+		BlueprintObj->GetRenderGrid()->Modify();
 	}
 }
 
@@ -316,7 +326,6 @@ void UE::RenderGrid::Private::FRenderGridEditor::SetSelectedRenderGridJobs(const
 	if (URenderGridJobSelection* JobSelection = RenderGridJobSelection.Get(); IsValid(JobSelection))
 	{
 		{
-			FScopedTransaction Transaction(LOCTEXT("ChangeJobSelection", "Change Job Selection"));
 			JobSelection->Modify();
 			if (!JobSelection->SetSelectedRenderGridJobs(Jobs))
 			{
@@ -414,6 +423,10 @@ void UE::RenderGrid::Private::FRenderGridEditor::OnBlueprintChangedImpl(UBluepri
 	if (InBlueprint)
 	{
 		RefreshInstance();
+		if (URenderGridBlueprint* RenderGridBlueprint = Cast<URenderGridBlueprint>(InBlueprint); IsValid(RenderGridBlueprint))
+		{
+			RenderGridBlueprint->PropagateAllPropertiesToInstances();
+		}
 	}
 }
 
@@ -431,9 +444,9 @@ void UE::RenderGrid::Private::FRenderGridEditor::NotifyPreChange(FProperty* Prop
 {
 	FBlueprintEditor::NotifyPreChange(PropertyAboutToChange);
 
-	if (URenderGridBlueprint* RenderGridBP = GetRenderGridBlueprint())
+	if (URenderGridBlueprint* RenderGridBlueprint = GetRenderGridBlueprint())
 	{
-		RenderGridBP->Modify();
+		RenderGridBlueprint->Modify();
 	}
 }
 
@@ -461,10 +474,14 @@ void UE::RenderGrid::Private::FRenderGridEditor::AddJobAction()
 		FScopedTransaction Transaction(LOCTEXT("AddJob", "Add Job"));
 		MarkAsModified();
 
-		if (URenderGridJob* Job = Instance->CreateAndAddNewRenderGridJob(); IsValid(Job))
+		if (URenderGridBlueprint* PreviewBlueprint = PreviewBlueprintWeakPtr.Get(); IsValid(PreviewBlueprint))
 		{
-			OnRenderGridJobCreated().Broadcast(Job);
-			OnRenderGridChanged().Broadcast();
+			if (URenderGridJob* Job = PreviewBlueprint->GetRenderGrid()->CreateAndAddNewRenderGridJob(); IsValid(Job))
+			{
+				PreviewBlueprint->PropagateJobsToInstances();
+				OnRenderGridJobCreated().Broadcast(Job);
+				OnRenderGridChanged().Broadcast();
+			}
 		}
 	}
 }
@@ -482,11 +499,15 @@ void UE::RenderGrid::Private::FRenderGridEditor::DuplicateJobAction()
 		FScopedTransaction Transaction(LOCTEXT("DuplicateJob", "Duplicate Job"));
 		MarkAsModified();
 
-		for (URenderGridJob* SelectedJob : SelectedJobs)
+		if (URenderGridBlueprint* PreviewBlueprint = PreviewBlueprintWeakPtr.Get(); IsValid(PreviewBlueprint))
 		{
-			Instance->DuplicateAndAddRenderGridJob(SelectedJob);
+			for (URenderGridJob* SelectedJob : SelectedJobs)
+			{
+				PreviewBlueprint->GetRenderGrid()->DuplicateAndAddRenderGridJob(SelectedJob);
+			}
+			PreviewBlueprint->PropagateJobsToInstances();
+			OnRenderGridChanged().Broadcast();
 		}
-		OnRenderGridChanged().Broadcast();
 	}
 }
 
@@ -498,27 +519,20 @@ void UE::RenderGrid::Private::FRenderGridEditor::DeleteJobAction()
 		return;
 	}
 
-	const FText TitleText = LOCTEXT("ConfirmToDeleteTitle", "Confirm To Delete");
-	const EAppReturnType::Type DialogResult = FMessageDialog::Open(
-		EAppMsgType::OkCancel,
-		((SelectedJobs.Num() == 1) ? LOCTEXT("ConfirmToDeleteSingleText", "Are you sure you want to delete the selected job?") : LOCTEXT("ConfirmToDeleteMultipleText", "Are you sure you want to delete the selected jobs?")),
-		&TitleText);
-
-	if (DialogResult != EAppReturnType::Ok)
-	{
-		return;
-	}
-
 	if (URenderGrid* Instance = GetInstance(); IsValid(Instance))
 	{
 		FScopedTransaction Transaction(LOCTEXT("DeleteJob", "Delete Job"));
 		MarkAsModified();
 
-		for (URenderGridJob* SelectedJob : SelectedJobs)
+		if (URenderGridBlueprint* PreviewBlueprint = PreviewBlueprintWeakPtr.Get(); IsValid(PreviewBlueprint))
 		{
-			Instance->RemoveRenderGridJob(SelectedJob);
+			for (URenderGridJob* SelectedJob : SelectedJobs)
+			{
+				PreviewBlueprint->GetRenderGrid()->RemoveRenderGridJob(SelectedJob);
+			}
+			PreviewBlueprint->PropagateJobsToInstances();
+			OnRenderGridChanged().Broadcast();
 		}
-		OnRenderGridChanged().Broadcast();
 	}
 }
 
@@ -626,42 +640,44 @@ void UE::RenderGrid::Private::FRenderGridEditor::FillToolbar(FToolBarBuilder& To
 
 void UE::RenderGrid::Private::FRenderGridEditor::DestroyInstance()
 {
-	if (URenderGrid* Instance = GetInstance(); IsValid(Instance))
+	if (URenderGrid* Instance = RenderGridWeakPtr.Get(); IsValid(Instance))
 	{
 		// Execute the blueprint event
 		Instance->EndEditor();
 
-		// Store the data
-		Instance->OnClose();
-
 		// Garbage collection
-		RenderGridWeakPtr.Reset();
 		Instance->MarkAsGarbage();
+	}
 
-		// Clear the undo/redo buffer, to prevent any bugs/glitches with undo/redo while the render grid editor is closed
-		GEditor->ResetTransaction(LOCTEXT("RenderGridEditorClosed", "Render Grid Editor Closed"));
+	if (URenderGridBlueprint* PreviewBlueprint = PreviewBlueprintWeakPtr.Get(); IsValid(PreviewBlueprint))
+	{
+		PreviewBlueprint->Save();
 	}
+
+	RenderGridWeakPtr.Reset();
+	PreviewBlueprintWeakPtr.Reset();
+
+	// Clear the undo/redo buffer, to prevent any bugs/glitches with undo/redo while the render grid editor is closed
+	//GEditor->ResetTransaction(LOCTEXT("RenderGridEditorClosed", "RenderGrid Editor Closed"));
 }
 
 void UE::RenderGrid::Private::FRenderGridEditor::UpdateInstance(UBlueprint* InBlueprint, bool bInForceFullUpdate)
 {
 	// If the Blueprint is changing
-	if ((InBlueprint != PreviewBlueprint.Get()) || bInForceFullUpdate)
+	if ((InBlueprint != PreviewBlueprintWeakPtr.Get()) || bInForceFullUpdate)
 	{
 		// Destroy the previous instance
 		DestroyInstance();
 
 		// Save the Blueprint we're creating a preview for
-		PreviewBlueprint = TStrongObjectPtr(Cast<URenderGridBlueprint>(InBlueprint));
+		PreviewBlueprintWeakPtr = Cast<URenderGridBlueprint>(InBlueprint);
 
-		URenderGrid* RenderGrid;
-		// Create the Widget, we have to do special swapping out of the widget tree.
+		URenderGrid* RenderGrid = nullptr;
+		if (URenderGridBlueprint* PreviewBlueprint = PreviewBlueprintWeakPtr.Get(); IsValid(PreviewBlueprint) && IsValid(PreviewBlueprint->GeneratedClass))
 		{
-			// Assign the outer to the game instance if it exists, otherwise use the world
-			{
-				FMakeClassSpawnableOnScope TemporarilySpawnable(PreviewBlueprint->GeneratedClass);
-				RenderGrid = NewObject<URenderGrid>(PreviewScene.GetWorld(), PreviewBlueprint->GeneratedClass);
-			}
+			PreviewBlueprint->Load();
+			FMakeClassSpawnableOnScope TemporarilySpawnable(PreviewBlueprint->GeneratedClass);
+			RenderGrid = NewObject<URenderGrid>(PreviewScene.GetWorld(), PreviewBlueprint->GeneratedClass);
 		}
 
 		// Set the debug object again, if it was debugging
@@ -675,7 +691,10 @@ void UE::RenderGrid::Private::FRenderGridEditor::UpdateInstance(UBlueprint* InBl
 		OnRenderGridJobsSelectionChanged().Broadcast();
 
 		// Execute the blueprint event
-		RenderGrid->BeginEditor();
+		if (IsValid(RenderGrid))
+		{
+			RenderGrid->BeginEditor();
+		}
 	}
 }
 
diff --git a/Engine/Plugins/Experimental/RenderGrid/Source/RenderGridEditor/Private/Toolkit/RenderGridEditor.h b/Engine/Plugins/Experimental/RenderGrid/Source/RenderGridEditor/Private/Toolkit/RenderGridEditor.h
--- a/Engine/Plugins/Experimental/RenderGrid/Source/RenderGridEditor/Private/Toolkit/RenderGridEditor.h
+++ b/Engine/Plugins/Experimental/RenderGrid/Source/RenderGridEditor/Private/Toolkit/RenderGridEditor.h
@@ -181,7 +181,7 @@ namespace UE::RenderGrid::Private
 		TSharedPtr<FExtender> ToolbarExtender;
 
 		/** The blueprint instance that's currently visible in the editor. */
-		TStrongObjectPtr<URenderGridBlueprint> PreviewBlueprint;
+		TWeakObjectPtr<URenderGridBlueprint> PreviewBlueprintWeakPtr;
 
 		/** The current render grid instance that's visible in the editor. */
 		mutable TWeakObjectPtr<URenderGrid> RenderGridWeakPtr;
diff --git a/Engine/Plugins/Experimental/RenderGrid/Source/RenderGridEditor/Private/UI/SRenderGrid.cpp b/Engine/Plugins/Experimental/RenderGrid/Source/RenderGridEditor/Private/UI/SRenderGrid.cpp
--- a/Engine/Plugins/Experimental/RenderGrid/Source/RenderGridEditor/Private/UI/SRenderGrid.cpp
+++ b/Engine/Plugins/Experimental/RenderGrid/Source/RenderGridEditor/Private/UI/SRenderGrid.cpp
@@ -5,6 +5,8 @@
 #include "RenderGrid/RenderGrid.h"
 #include "Kismet2/BlueprintEditorUtils.h"
 #include "SlateOptMacros.h"
+#include "Blueprints/RenderGridBlueprint.h"
+#include "ScopedTransaction.h"
 
 #define LOCTEXT_NAMESPACE "SRenderGrid"
 
@@ -76,7 +78,9 @@ void UE::RenderGrid::Private::SRenderGrid::NotifyPostChange(const FPropertyChang
 {
 	if (const TSharedPtr<IRenderGridEditor> BlueprintEditor = BlueprintEditorWeakPtr.Pin())
 	{
+		FScopedTransaction Transaction(LOCTEXT("RenderGridParameterChanged", "RenderGrid Parameter Changed"));
 		BlueprintEditor->MarkAsModified();
+		BlueprintEditor->GetRenderGridBlueprint()->PropagateAllPropertiesExceptJobsToAsset(BlueprintEditor->GetInstance());
 		BlueprintEditor->OnRenderGridChanged().Broadcast();
 	}
 }
diff --git a/Engine/Plugins/Experimental/RenderGrid/Source/RenderGridEditor/Private/UI/SRenderGridJobList.cpp b/Engine/Plugins/Experimental/RenderGrid/Source/RenderGridEditor/Private/UI/SRenderGridJobList.cpp
--- a/Engine/Plugins/Experimental/RenderGrid/Source/RenderGridEditor/Private/UI/SRenderGridJobList.cpp
+++ b/Engine/Plugins/Experimental/RenderGrid/Source/RenderGridEditor/Private/UI/SRenderGridJobList.cpp
@@ -21,6 +21,7 @@
 #include "PropertyCustomizationHelpers.h"
 #include "ScopedTransaction.h"
 #include "SlateOptMacros.h"
+#include "Blueprints/RenderGridBlueprint.h"
 #include "Framework/Application/SlateApplication.h"
 #include "Styling/AppStyle.h"
 #include "Tracks/MovieSceneSubTrack.h"
@@ -400,6 +401,7 @@ void UE::RenderGrid::Private::SRenderGridJobListTable::HandleJobListSelectionCha
 	}
 	if (const TSharedPtr<IRenderGridEditor> BlueprintEditor = BlueprintEditorWeakPtr.Pin())
 	{
+		FScopedTransaction Transaction(LOCTEXT("ChangeJobSelection", "Change Job Selection"));
 		BlueprintEditor->SetSelectedRenderGridJobs(GetSelectedItems());
 	}
 }
@@ -447,9 +449,11 @@ FReply UE::RenderGrid::Private::SRenderGridJobListTableRow::OnAcceptDrop(const F
 		{
 			if (URenderGrid* Instance = BlueprintEditor->GetInstance(); IsValid(Instance))
 			{
+				FScopedTransaction Transaction(LOCTEXT("MoveJob", "Move Job"));
+				BlueprintEditor->MarkAsModified();
 				if (Instance->ReorderRenderGridJob(DragDropOp->GetJob(), InJob, (InItemDropZone != EItemDropZone::AboveItem)))
 				{
-					BlueprintEditor->MarkAsModified();
+					BlueprintEditor->GetRenderGridBlueprint()->PropagateJobsToAsset(Instance);
 					BlueprintEditor->OnRenderGridChanged().Broadcast();
 					return FReply::Handled();
 				}
