RenderGrid plugin:
- Moved the render grid settings to a separate object, so it can be separated in the UI.
- Added a blueprint to obtain all RenderGrids from the project.
- Removed the "no jobs to render" error message popup, since jobs can be automatically generated during rendering.
- Reworked the rendering queue functionality, to be able to handle pausing, canceling, etc all correctly.
- Added a blueprint to be able to render rendergrids programmatically.
- Fixed several issues that would cause crashes in some situations (like when you close the render grid editor and the render queue on-finish callback would fire, causing a nullpointer exception and crash).
- Fixed an issue where the blueprint graph didn't work during rendering.
- Made the RenderGridQueue automatically queue up RenderGridQueues when multiple are being executed at the same time.
- Made render grids only render using the CDO (as this will never be destroyed).
- Added blueprints to allow rendering a single or multiple given jobs (instead of only being able to render all enabled jobs).
- Make the UI now always hidden if it's rendering anything, even if it's not rendering the currently opened render grid.
- Made the job selection static, to prevent multiple jobs (in different render grid editors) from being opened at the same time, as that would cause bugs with the remote control properties.

diff --git a/Engine/Plugins/Experimental/RenderGrid/Source/RenderGrid/Private/RenderGrid/RenderGrid.cpp b/Engine/Plugins/Experimental/RenderGrid/Source/RenderGrid/Private/RenderGrid/RenderGrid.cpp
--- a/Engine/Plugins/Experimental/RenderGrid/Source/RenderGrid/Private/RenderGrid/RenderGrid.cpp
+++ b/Engine/Plugins/Experimental/RenderGrid/Source/RenderGrid/Private/RenderGrid/RenderGrid.cpp
@@ -9,9 +9,21 @@
 #include "MoviePipelinePrimaryConfig.h"
 #include "MoviePipelineOutputSetting.h"
 #include "MovieScene.h"
+#include "RenderGrid/RenderGridQueue.h"
 #include "UObject/ObjectSaveContext.h"
 
 
+URenderGridSettings::URenderGridSettings()
+	: PropsSourceType(ERenderGridPropsSourceType::RemoteControl)
+	, PropsSourceOrigin_RemoteControl(nullptr)
+	, CachedPropsSource(nullptr)
+	, CachedPropsSourceType(ERenderGridPropsSourceType::Local)
+	, CachedPropsSourceOriginWeakPtr(nullptr)
+{
+	SetFlags(RF_Public | RF_Transactional);
+}
+
+
 URenderGridDefaults::URenderGridDefaults()
 	: LevelSequence(nullptr)
 	, RenderPreset(nullptr)
@@ -47,14 +59,14 @@ TOptional<int32> URenderGridJob::GetSequenceStartFrame() const
 	const UMovieScene* MovieScene = LevelSequence->MovieScene;
 	const FFrameRate TickResolution = MovieScene->GetTickResolution();
 	const FFrameRate DisplayRate = MovieScene->GetDisplayRate();
-	const UMoviePipelineOutputSetting* Settings = GetRenderPresetOutputSettings();
+	const UMoviePipelineOutputSetting* MovieOutputSettings = GetRenderPresetOutputSettings();
 
-	if (bOverrideStartFrame || (IsValid(Settings) && Settings->bUseCustomPlaybackRange))
+	if (bOverrideStartFrame || (IsValid(MovieOutputSettings) && MovieOutputSettings->bUseCustomPlaybackRange))
 	{
-		int32 Frame = (bOverrideStartFrame ? CustomStartFrame : Settings->CustomStartFrame);
-		if (IsValid(Settings) && Settings->bUseCustomFrameRate)
+		int32 Frame = (bOverrideStartFrame ? CustomStartFrame : MovieOutputSettings->CustomStartFrame);
+		if (IsValid(MovieOutputSettings) && MovieOutputSettings->bUseCustomFrameRate)
 		{
-			return FMath::FloorToInt(Frame / (Settings->OutputFrameRate / DisplayRate).AsDecimal());
+			return FMath::FloorToInt(Frame / (MovieOutputSettings->OutputFrameRate / DisplayRate).AsDecimal());
 		}
 		return Frame;
 	}
@@ -73,14 +85,14 @@ TOptional<int32> URenderGridJob::GetSequenceEndFrame() const
 	const UMovieScene* MovieScene = LevelSequence->MovieScene;
 	const FFrameRate TickResolution = MovieScene->GetTickResolution();
 	const FFrameRate DisplayRate = MovieScene->GetDisplayRate();
-	const UMoviePipelineOutputSetting* Settings = GetRenderPresetOutputSettings();
+	const UMoviePipelineOutputSetting* MovieOutputSettings = GetRenderPresetOutputSettings();
 
-	if (bOverrideEndFrame || (IsValid(Settings) && Settings->bUseCustomPlaybackRange))
+	if (bOverrideEndFrame || (IsValid(MovieOutputSettings) && MovieOutputSettings->bUseCustomPlaybackRange))
 	{
-		int32 Frame = (bOverrideEndFrame ? CustomEndFrame : Settings->CustomEndFrame);
-		if (IsValid(Settings) && Settings->bUseCustomFrameRate)
+		int32 Frame = (bOverrideEndFrame ? CustomEndFrame : MovieOutputSettings->CustomEndFrame);
+		if (IsValid(MovieOutputSettings) && MovieOutputSettings->bUseCustomFrameRate)
 		{
-			return FMath::FloorToInt(Frame / (Settings->OutputFrameRate / DisplayRate).AsDecimal());
+			return FMath::FloorToInt(Frame / (MovieOutputSettings->OutputFrameRate / DisplayRate).AsDecimal());
 		}
 		return Frame;
 	}
@@ -156,10 +168,10 @@ TOptional<int32> URenderGridJob::GetStartFrame() const
 		return CustomStartFrame;
 	}
 
-	const UMoviePipelineOutputSetting* Settings = GetRenderPresetOutputSettings();
-	if (IsValid(Settings) && Settings->bUseCustomPlaybackRange)
+	const UMoviePipelineOutputSetting* MovieOutputSettings = GetRenderPresetOutputSettings();
+	if (IsValid(MovieOutputSettings) && MovieOutputSettings->bUseCustomPlaybackRange)
 	{
-		return Settings->CustomStartFrame;
+		return MovieOutputSettings->CustomStartFrame;
 	}
 
 	if (!IsValid(LevelSequence))
@@ -170,9 +182,9 @@ TOptional<int32> URenderGridJob::GetStartFrame() const
 	const UMovieScene* MovieScene = LevelSequence->MovieScene;
 	const FFrameRate TickResolution = MovieScene->GetTickResolution();
 	FFrameRate DisplayRate = MovieScene->GetDisplayRate();
-	if (IsValid(Settings) && Settings->bUseCustomFrameRate)
+	if (IsValid(MovieOutputSettings) && MovieOutputSettings->bUseCustomFrameRate)
 	{
-		DisplayRate = Settings->OutputFrameRate;
+		DisplayRate = MovieOutputSettings->OutputFrameRate;
 	}
 	const int32 StartFrameNumber = MovieScene->GetPlaybackRange().GetLowerBoundValue().Value;
 	return FMath::FloorToInt(StartFrameNumber / (TickResolution / DisplayRate).AsDecimal());
@@ -185,10 +197,10 @@ TOptional<int32> URenderGridJob::GetEndFrame() const
 		return CustomEndFrame;
 	}
 
-	const UMoviePipelineOutputSetting* Settings = GetRenderPresetOutputSettings();
-	if (IsValid(Settings) && Settings->bUseCustomPlaybackRange)
+	const UMoviePipelineOutputSetting* MovieOutputSettings = GetRenderPresetOutputSettings();
+	if (IsValid(MovieOutputSettings) && MovieOutputSettings->bUseCustomPlaybackRange)
 	{
-		return Settings->CustomEndFrame;
+		return MovieOutputSettings->CustomEndFrame;
 	}
 
 	if (!IsValid(LevelSequence))
@@ -199,9 +211,9 @@ TOptional<int32> URenderGridJob::GetEndFrame() const
 	const UMovieScene* MovieScene = LevelSequence->MovieScene;
 	const FFrameRate TickResolution = MovieScene->GetTickResolution();
 	FFrameRate DisplayRate = MovieScene->GetDisplayRate();
-	if (IsValid(Settings) && Settings->bUseCustomFrameRate)
+	if (IsValid(MovieOutputSettings) && MovieOutputSettings->bUseCustomFrameRate)
 	{
-		DisplayRate = Settings->OutputFrameRate;
+		DisplayRate = MovieOutputSettings->OutputFrameRate;
 	}
 	const int32 EndFrameNumber = MovieScene->GetPlaybackRange().GetUpperBoundValue().Value;
 	return FMath::FloorToInt(EndFrameNumber / (TickResolution / DisplayRate).AsDecimal());
@@ -220,13 +232,13 @@ TOptional<double> URenderGridJob::GetStartTime() const
 		return TOptional<double>();
 	}
 
-	const UMoviePipelineOutputSetting* Settings = GetRenderPresetOutputSettings();
+	const UMoviePipelineOutputSetting* MovieOutputSettings = GetRenderPresetOutputSettings();
 
 	const UMovieScene* MovieScene = LevelSequence->MovieScene;
 	FFrameRate DisplayRate = MovieScene->GetDisplayRate();
-	if (IsValid(Settings) && Settings->bUseCustomFrameRate)
+	if (IsValid(MovieOutputSettings) && MovieOutputSettings->bUseCustomFrameRate)
 	{
-		DisplayRate = Settings->OutputFrameRate;
+		DisplayRate = MovieOutputSettings->OutputFrameRate;
 	}
 	return StartFrame.Get(0) / DisplayRate.AsDecimal();
 }
@@ -244,13 +256,13 @@ TOptional<double> URenderGridJob::GetEndTime() const
 		return TOptional<double>();
 	}
 
-	const UMoviePipelineOutputSetting* Settings = GetRenderPresetOutputSettings();
+	const UMoviePipelineOutputSetting* MovieOutputSettings = GetRenderPresetOutputSettings();
 
 	const UMovieScene* MovieScene = LevelSequence->MovieScene;
 	FFrameRate DisplayRate = MovieScene->GetDisplayRate();
-	if (IsValid(Settings) && Settings->bUseCustomFrameRate)
+	if (IsValid(MovieOutputSettings) && MovieOutputSettings->bUseCustomFrameRate)
 	{
-		DisplayRate = Settings->OutputFrameRate;
+		DisplayRate = MovieOutputSettings->OutputFrameRate;
 	}
 	return EndFrame.Get(0) / DisplayRate.AsDecimal();
 }
@@ -269,13 +281,13 @@ TOptional<double> URenderGridJob::GetDurationInSeconds() const
 		return TOptional<double>();
 	}
 
-	const UMoviePipelineOutputSetting* Settings = GetRenderPresetOutputSettings();
+	const UMoviePipelineOutputSetting* MovieOutputSettings = GetRenderPresetOutputSettings();
 
 	const UMovieScene* MovieScene = LevelSequence->MovieScene;
 	FFrameRate DisplayRate = MovieScene->GetDisplayRate();
-	if (IsValid(Settings) && Settings->bUseCustomFrameRate)
+	if (IsValid(MovieOutputSettings) && MovieOutputSettings->bUseCustomFrameRate)
 	{
-		DisplayRate = Settings->OutputFrameRate;
+		DisplayRate = MovieOutputSettings->OutputFrameRate;
 	}
 	return (EndFrame.Get(0) - StartFrame.Get(0)) / DisplayRate.AsDecimal();
 }
@@ -286,19 +298,19 @@ FIntPoint URenderGridJob::GetOutputResolution() const
 	{
 		return CustomResolution;
 	}
-	if (UMoviePipelineOutputSetting* Settings = GetRenderPresetOutputSettings())
+	if (UMoviePipelineOutputSetting* MovieOutputSettings = GetRenderPresetOutputSettings())
 	{
-		return Settings->OutputResolution;
+		return MovieOutputSettings->OutputResolution;
 	}
 	return GetDefault<UMoviePipelineOutputSetting>()->OutputResolution;
 }
 
 double URenderGridJob::GetOutputAspectRatio() const
 {
-	const UMoviePipelineOutputSetting* Settings = GetRenderPresetOutputSettings();
-	if (IsValid(Settings))
+	const UMoviePipelineOutputSetting* MovieOutputSettings = GetRenderPresetOutputSettings();
+	if (IsValid(MovieOutputSettings))
 	{
-		return static_cast<double>(Settings->OutputResolution.X) / static_cast<double>(Settings->OutputResolution.Y);
+		return static_cast<double>(MovieOutputSettings->OutputResolution.X) / static_cast<double>(MovieOutputSettings->OutputResolution.Y);
 	}
 
 	const UMoviePipelineOutputSetting* DefaultSettings = GetDefault<UMoviePipelineOutputSetting>();
@@ -354,19 +366,19 @@ UMoviePipelineOutputSetting* URenderGridJob::GetRenderPresetOutputSettings() con
 	{
 		return nullptr;
 	}
-	for (UMoviePipelineSetting* Settings : RenderPreset->FindSettingsByClass(UMoviePipelineOutputSetting::StaticClass(), false))
+	for (UMoviePipelineSetting* MovieSettings : RenderPreset->FindSettingsByClass(UMoviePipelineOutputSetting::StaticClass(), false))
 	{
-		if (!IsValid(Settings))
+		if (!IsValid(MovieSettings))
 		{
 			continue;
 		}
-		if (UMoviePipelineOutputSetting* OutputSettings = Cast<UMoviePipelineOutputSetting>(Settings))
+		if (UMoviePipelineOutputSetting* MovieOutputSettings = Cast<UMoviePipelineOutputSetting>(MovieSettings))
 		{
-			if (!OutputSettings->IsEnabled())
+			if (!MovieOutputSettings->IsEnabled())
 			{
 				continue;
 			}
-			return OutputSettings;
+			return MovieOutputSettings;
 		}
 	}
 	return nullptr;
@@ -435,12 +447,9 @@ bool URenderGridJob::SetRemoteControlValue(const TSharedPtr<FRemoteControlEntity
 
 URenderGrid::URenderGrid()
 	: Guid(FGuid::NewGuid())
-	, PropsSourceType(ERenderGridPropsSourceType::RemoteControl)
-	, PropsSourceOrigin_RemoteControl(nullptr)
+	, Settings(CreateDefaultSubobject<URenderGridSettings>(TEXT("RenderGridSettings")))
 	, Defaults(CreateDefaultSubobject<URenderGridDefaults>(TEXT("RenderGridDefaults")))
 	, bExecutingBlueprintEvent(false)
-	, CachedPropsSource(nullptr)
-	, CachedPropsSourceType(ERenderGridPropsSourceType::Local)
 {
 	SetFlags(RF_Public);
 	if (!HasAnyFlags(RF_ClassDefaultObject | RF_DefaultSubObject))
@@ -451,12 +460,6 @@ URenderGrid::URenderGrid()
 
 UWorld* URenderGrid::GetWorld() const
 {
-	if (HasAllFlags(RF_ClassDefaultObject))
-	{
-		// If we are a CDO, we must return nullptr instead of calling Outer->GetWorld() to fool UObject::ImplementsGetWorld.
-		return nullptr;
-	}
-
 	for (const FWorldContext& Context : GEngine->GetWorldContexts())
 	{
 		if (Context.WorldType == EWorldType::PIE)
@@ -474,11 +477,15 @@ UWorld* URenderGrid::GetWorld() const
 	{
 		return CachedWorld;
 	}
-	UObject* Outer = GetOuter();// Could be a GameInstance, could be World, could also be a WidgetTree, so we're just going to follow the outer chain to find the world we're in.
-	while (Outer)
+	UObject* Outer = GetOuter(); // Could be a GameInstance, could be World, could also be a WidgetTree, so we're just going to follow the outer chain to find the world we're in.
+	while (IsValid(Outer))
 	{
-		UWorld* World = Outer->GetWorld();
-		if (IsValid(World))
+		if (UWorld* World = Cast<UWorld>(Outer); IsValid(World))
+		{
+			CachedWorldWeakPtr = World;
+			return World;
+		}
+		if (UWorld* World = Outer->GetWorld(); IsValid(World))
 		{
 			CachedWorldWeakPtr = World;
 			return World;
@@ -497,12 +504,24 @@ void URenderGrid::PostLoad()
 {
 	Super::PostLoad();
 
-	if (PropsSourceType == ERenderGridPropsSourceType::Local)
+	SetFlags(RF_Public);
+
+	if (!IsValid(Settings))
+	{
+		Settings = NewObject<URenderGridSettings>(this, TEXT("RenderGridSettings"));
+	}
+	Settings->SetFlags(RF_Public | RF_Transactional);
+	if (Settings->PropsSourceType == ERenderGridPropsSourceType::Local)
 	{
-		PropsSourceType = ERenderGridPropsSourceType::RemoteControl;
+		Settings->PropsSourceType = ERenderGridPropsSourceType::RemoteControl;
 	}
 
-	SetFlags(RF_Public);
+	if (!IsValid(Defaults))
+	{
+		Defaults = NewObject<URenderGridDefaults>(this, TEXT("RenderGridDefaults"));
+	}
+	Defaults->SetFlags(RF_Public | RF_Transactional);
+
 	if (!HasAnyFlags(RF_ClassDefaultObject | RF_DefaultSubObject))
 	{
 		ClearFlags(RF_Transactional);
@@ -514,11 +533,6 @@ void URenderGrid::PostLoad()
 			Job->SetFlags(RF_Public | RF_Transactional);
 		}
 	}
-	if (!IsValid(Defaults))
-	{
-		Defaults = NewObject<URenderGridDefaults>(this, TEXT("RenderGridDefaults"));
-	}
-	Defaults->SetFlags(RF_Public | RF_Transactional);
 }
 
 void URenderGrid::CopyJobs(URenderGrid* From)
@@ -618,34 +632,73 @@ void URenderGrid::EndViewportRender(URenderGridJob* Job)
 	bExecutingBlueprintEvent = false;
 }
 
+URenderGridQueue* URenderGrid::Render()
+{
+	return RenderJobs(GetEnabledRenderGridJobs());
+}
+
+URenderGridQueue* URenderGrid::RenderJobs(const TArray<URenderGridJob*>& Jobs)
+{
+	if (!HasAnyFlags(RF_ClassDefaultObject))
+	{
+		if (URenderGrid* DefaultObject = Cast<URenderGrid>(GetClass()->GetDefaultObject(true)); IsValid(DefaultObject))
+		{
+			CopyAllProperties(DefaultObject);
+			return DefaultObject->RenderJobs(Jobs);
+		}
+		return nullptr;
+	}
+
+	TArray<URenderGridJob*> JobsToRender;
+	for (URenderGridJob* Job : Jobs)
+	{
+		if (IsValid(Job) && HasRenderGridJob(Job))
+		{
+			JobsToRender.Add(Job);
+		}
+	}
+	if (URenderGridQueue* RenderQueue = UE::RenderGrid::IRenderGridModule::Get().GetManager().CreateBatchRenderQueue(this, Jobs); IsValid(RenderQueue))
+	{
+		RenderQueue->Execute();
+		return RenderQueue;
+	}
+	return nullptr;
+}
+
+URenderGridQueue* URenderGrid::RenderJob(URenderGridJob* Job)
+{
+	return RenderJobs({Job});
+}
+
 void URenderGrid::SetPropsSource(ERenderGridPropsSourceType InPropsSourceType, UObject* InPropsSourceOrigin)
 {
 	if (InPropsSourceType == ERenderGridPropsSourceType::RemoteControl)
 	{
-		PropsSourceType = InPropsSourceType;
-		PropsSourceOrigin_RemoteControl = Cast<URemoteControlPreset>(InPropsSourceOrigin);
+		Settings->PropsSourceType = InPropsSourceType;
+		Settings->PropsSourceOrigin_RemoteControl = Cast<URemoteControlPreset>(InPropsSourceOrigin);
 		return;
 	}
-	PropsSourceType = ERenderGridPropsSourceType::RemoteControl;
+	Settings->PropsSourceType = ERenderGridPropsSourceType::RemoteControl;
+	Settings->PropsSourceOrigin_RemoteControl = nullptr;
 }
 
 URenderGridPropsSourceBase* URenderGrid::GetPropsSource() const
 {
 	UObject* PropsSourceOrigin = GetPropsSourceOrigin();
-	if (!IsValid(CachedPropsSource) || (CachedPropsSourceType != PropsSourceType) || (CachedPropsSourceOriginWeakPtr.Get() != PropsSourceOrigin))
+	if (!IsValid(Settings->CachedPropsSource) || (Settings->CachedPropsSourceType != Settings->PropsSourceType) || (Settings->CachedPropsSourceOriginWeakPtr.Get() != PropsSourceOrigin))
 	{
-		CachedPropsSourceType = PropsSourceType;
-		CachedPropsSourceOriginWeakPtr = PropsSourceOrigin;
-		CachedPropsSource = UE::RenderGrid::IRenderGridModule::Get().CreatePropsSource(const_cast<URenderGrid*>(this), PropsSourceType, PropsSourceOrigin);
+		Settings->CachedPropsSourceType = Settings->PropsSourceType;
+		Settings->CachedPropsSourceOriginWeakPtr = PropsSourceOrigin;
+		Settings->CachedPropsSource = UE::RenderGrid::IRenderGridModule::Get().CreatePropsSource(const_cast<URenderGrid*>(this), Settings->PropsSourceType, PropsSourceOrigin);
 	}
-	return CachedPropsSource;
+	return Settings->CachedPropsSource;
 }
 
 UObject* URenderGrid::GetPropsSourceOrigin() const
 {
-	if (PropsSourceType == ERenderGridPropsSourceType::RemoteControl)
+	if (Settings->PropsSourceType == ERenderGridPropsSourceType::RemoteControl)
 	{
-		return PropsSourceOrigin_RemoteControl;
+		return Settings->PropsSourceOrigin_RemoteControl;
 	}
 	return nullptr;
 }
@@ -656,19 +709,19 @@ UMoviePipelineOutputSetting* URenderGrid::GetDefaultRenderPresetOutputSettings()
 	{
 		return nullptr;
 	}
-	for (UMoviePipelineSetting* Settings : Defaults->RenderPreset->FindSettingsByClass(UMoviePipelineOutputSetting::StaticClass(), false))
+	for (UMoviePipelineSetting* MovieSettings : Defaults->RenderPreset->FindSettingsByClass(UMoviePipelineOutputSetting::StaticClass(), false))
 	{
-		if (!IsValid(Settings))
+		if (!IsValid(MovieSettings))
 		{
 			continue;
 		}
-		if (UMoviePipelineOutputSetting* OutputSettings = Cast<UMoviePipelineOutputSetting>(Settings))
+		if (UMoviePipelineOutputSetting* MovieOutputSettings = Cast<UMoviePipelineOutputSetting>(MovieSettings))
 		{
-			if (!OutputSettings->IsEnabled())
+			if (!MovieOutputSettings->IsEnabled())
 			{
 				continue;
 			}
-			return OutputSettings;
+			return MovieOutputSettings;
 		}
 	}
 	return nullptr;
diff --git a/Engine/Plugins/Experimental/RenderGrid/Source/RenderGrid/Private/RenderGrid/RenderGridManager.cpp b/Engine/Plugins/Experimental/RenderGrid/Source/RenderGrid/Private/RenderGrid/RenderGridManager.cpp
--- a/Engine/Plugins/Experimental/RenderGrid/Source/RenderGrid/Private/RenderGrid/RenderGridManager.cpp
+++ b/Engine/Plugins/Experimental/RenderGrid/Source/RenderGrid/Private/RenderGrid/RenderGridManager.cpp
@@ -7,11 +7,11 @@
 #include "MoviePipelinePIEExecutor.h"
 
 
-URenderGridQueue* UE::RenderGrid::FRenderGridManager::CreateBatchRenderQueue(URenderGrid* Grid)
+URenderGridQueue* UE::RenderGrid::FRenderGridManager::CreateBatchRenderQueue(URenderGrid* Grid, const TArray<URenderGridJob*>& Jobs)
 {
 	FRenderGridQueueCreateArgs JobArgs;
 	JobArgs.RenderGrid = TStrongObjectPtr(Grid);
-	JobArgs.RenderGridJobs.Append(Grid->GetEnabledRenderGridJobs());
+	JobArgs.RenderGridJobs.Append(Jobs);
 	JobArgs.bIsBatchRender = true;
 	URenderGridQueue* NewRenderQueue = URenderGridQueue::Create(JobArgs);
 	if (!IsValid(NewRenderQueue))
@@ -51,7 +51,7 @@ URenderGridQueue* UE::RenderGrid::FRenderGridManager::RenderPreviewFrame(const F
 
 	if (Args.Frame.IsSet())
 	{
-		constexpr int32 RenderFramesCount = 1;// can be more than 1 to prevent rendering issues, will always take the last frame that's rendered
+		constexpr int32 RenderFramesCount = 1; // can be more than 1 to prevent rendering issues, will always take the last frame that's rendered
 
 		JobCopy->SetIsUsingCustomStartFrame(true);
 		JobCopy->SetCustomStartFrame(Args.Frame.Get(0));
diff --git a/Engine/Plugins/Experimental/RenderGrid/Source/RenderGrid/Private/RenderGrid/RenderGridQueue.cpp b/Engine/Plugins/Experimental/RenderGrid/Source/RenderGrid/Private/RenderGrid/RenderGridQueue.cpp
--- a/Engine/Plugins/Experimental/RenderGrid/Source/RenderGrid/Private/RenderGrid/RenderGridQueue.cpp
+++ b/Engine/Plugins/Experimental/RenderGrid/Source/RenderGrid/Private/RenderGrid/RenderGridQueue.cpp
@@ -137,7 +137,7 @@ URenderGridMoviePipelineRenderJob* URenderGridMoviePipelineRenderJob::Create(URe
 		// anti-aliasing settings already present (and enabled)
 		if ((UE::MovieRenderPipeline::GetEffectiveAntiAliasingMethod(ExistingAntiAliasingSettings) == AAM_FXAA) && (ExistingAntiAliasingSettings->SpatialSampleCount <= 1) && (ExistingAntiAliasingSettings->TemporalSampleCount <= 1))
 		{
-			ExistingAntiAliasingSettings->TemporalSampleCount = 2;// FXAA transparency fix
+			ExistingAntiAliasingSettings->TemporalSampleCount = 2; // FXAA transparency fix
 		}
 	}
 	else if (UMoviePipelineAntiAliasingSetting* AntiAliasingSettings = Cast<UMoviePipelineAntiAliasingSetting>(NewJob->GetConfiguration()->FindOrAddSettingByClass(UMoviePipelineAntiAliasingSetting::StaticClass())))
@@ -303,6 +303,75 @@ void URenderGridMoviePipelineRenderJob::ExecuteFinished(UMoviePipelineExecutorBa
 }
 
 
+TQueue<TObjectPtr<URenderGridQueue>> URenderGridQueue::ExecutingQueues;
+URenderGridQueue::FOnExecutionQueueChanged URenderGridQueue::OnExecutionQueueChangedDelegate;
+
+bool URenderGridQueue::IsExecutingAny()
+{
+	return (GetCurrentlyExecutingQueue() != nullptr);
+}
+
+URenderGridQueue* URenderGridQueue::GetCurrentlyExecutingQueue()
+{
+	bool bRemovedAny = false;
+	while (TObjectPtr<URenderGridQueue>* CurrentlyExecutingQueue = ExecutingQueues.Peek())
+	{
+		URenderGridQueue* Queue = CurrentlyExecutingQueue->Get();
+		if (!IsValid(Queue))
+		{
+			ExecutingQueues.Pop();
+			bRemovedAny = true;
+			continue;
+		}
+		if (bRemovedAny)
+		{
+			if (!Queue->bPaused)
+			{
+				Queue->Queue->Start();
+			}
+			OnExecutionQueueChanged().Broadcast();
+		}
+		return Queue;
+	}
+	return nullptr;
+}
+
+void URenderGridQueue::AddExecutingQueue(URenderGridQueue* Queue)
+{
+	ExecutingQueues.Enqueue(Queue);
+	if ((GetCurrentlyExecutingQueue() == Queue) && !Queue->bPaused)
+	{
+		Queue->Queue->Start(); // can be called twice here, which is fine (once in GetCurrentlyExecutingQueue() if the previous Queue became invalid, and then once again here)
+	}
+	OnExecutionQueueChanged().Broadcast();
+}
+
+void URenderGridQueue::DoNextExecutingQueue()
+{
+	ExecutingQueues.Pop();
+	if (URenderGridQueue* Queue = GetCurrentlyExecutingQueue())
+	{
+		if (!Queue->bPaused)
+		{
+			Queue->Queue->Start(); // can be called twice here, which is fine (once in GetCurrentlyExecutingQueue() if the previous Queue became invalid, and then once again here)
+		}
+	}
+	OnExecutionQueueChanged().Broadcast();
+}
+
+
+void URenderGridQueue::Tick(float DeltaTime)
+{
+	if (bStarted && bFinished)
+	{
+		if (OnExecuteFinishedDelegate.IsBound())
+		{
+			OnExecuteFinishedDelegate.Broadcast(this, !bCanceled);
+			OnExecuteFinishedDelegate.Clear();
+		}
+	}
+}
+
 URenderGridQueue* URenderGridQueue::Create(const UE::RenderGrid::FRenderGridQueueCreateArgs& Args)
 {
 	URenderGrid* RenderGrid = Args.RenderGrid.Get();
@@ -315,7 +384,9 @@ URenderGridQueue* URenderGridQueue::Create(const UE::RenderGrid::FRenderGridQueu
 	RenderQueue->Args = Args;
 	RenderQueue->Queue = MakeShareable(new UE::RenderGrid::Private::FRenderGridGenericExecutionQueue);
 	RenderQueue->RenderGrid = RenderGrid;
+	RenderQueue->bStarted = false;
 	RenderQueue->bCanceled = false;
+	RenderQueue->bFinished = false;
 
 	for (const TStrongObjectPtr<URenderGridJob>& JobPtr : Args.RenderGridJobs)
 	{
@@ -329,18 +400,13 @@ URenderGridQueue* URenderGridQueue::Create(const UE::RenderGrid::FRenderGridQueu
 
 void URenderGridQueue::Execute()
 {
-	if (Queue->IsRunning())
-	{
-		return;
-	}
-
 	OnStart();
 }
 
 
 void URenderGridQueue::AddJob(URenderGridJob* Job)
 {
-	if (!IsValid(Job) || bCanceled)
+	if (!IsValid(Job) || bCanceled || bFinished)
 	{
 		return;
 	}
@@ -354,17 +420,25 @@ void URenderGridQueue::AddJob(URenderGridJob* Job)
 
 void URenderGridQueue::Pause()
 {
-	Queue->Stop();
+	bPaused = true;
+	if (bStarted && !bCanceled && !bFinished)
+	{
+		Queue->Stop();
+	}
 }
 
 void URenderGridQueue::Resume()
 {
-	Queue->Start();
+	bPaused = false;
+	if (bStarted && !bCanceled && !bFinished)
+	{
+		Queue->Start();
+	}
 }
 
 void URenderGridQueue::Cancel()
 {
-	if (bCanceled)
+	if (bCanceled || bFinished)
 	{
 		return;
 	}
@@ -439,7 +513,7 @@ int32 URenderGridQueue::GetJobsCount() const
 
 int32 URenderGridQueue::GetJobsRemainingCount() const
 {
-	if (bCanceled)
+	if (bCanceled || bFinished)
 	{
 		return 0;
 	}
@@ -468,7 +542,7 @@ int32 URenderGridQueue::GetJobsCompletedCount() const
 
 float URenderGridQueue::GetStatusPercentage() const
 {
-	if (bCanceled)
+	if (bCanceled || bFinished)
 	{
 		return 100.0;
 	}
@@ -484,17 +558,36 @@ float URenderGridQueue::GetStatusPercentage() const
 
 FString URenderGridQueue::GetStatus() const
 {
-	if (bCanceled)
+	if (bCanceled || bFinished)
 	{
 		return TEXT("Done");
 	}
+	if (!bStarted)
+	{
+		return TEXT("Waiting to be executed");
+	}
+	if (!IsCurrentlyRendering())
+	{
+		return TEXT("Waiting in queue");
+	}
+	if (bPaused)
+	{
+		return TEXT("Paused");
+	}
 	return TEXT("Rendering...");
 }
 
 
 void URenderGridQueue::OnStart()
 {
+	if (bStarted || bCanceled || bFinished)
+	{
+		return;
+	}
+	bStarted = true;
+
 	OnExecuteStartedDelegate.Broadcast(this);
+	OnExecuteStartedDelegate.Clear();
 	AddToRoot();
 
 	Queue->Add(UE::RenderGrid::Private::FRenderGridGenericExecutionQueueAction::CreateLambda([this]()
@@ -507,19 +600,22 @@ void URenderGridQueue::OnStart()
 	{
 		Queue->Add(UE::RenderGrid::Private::FRenderGridGenericExecutionQueueAction::CreateLambda([this]()
 		{
-			RenderGrid->BeginBatchRender(this);
+			if (IsValid(RenderGrid))
+			{
+				RenderGrid->BeginBatchRender(this);
+			}
 		}));
 		Queue->DelayFrames(2);
 	}
 
-	OnProcessJob(nullptr);// starts the job execution queueing
+	OnProcessJob(nullptr); // starts the job execution queueing
 
-	Queue->Start();
+	AddExecutingQueue(this);
 }
 
 void URenderGridQueue::OnProcessJob(URenderGridJob* Job)
 {
-	if (IsValid(Job) && !bCanceled)
+	if (IsValid(Job) && !bCanceled && !bFinished)
 	{
 		if (TObjectPtr<URenderGridMoviePipelineRenderJob>* EntryPtr = Entries.Find(Job))
 		{
@@ -527,13 +623,19 @@ void URenderGridQueue::OnProcessJob(URenderGridJob* Job)
 			{
 				Queue->Add(UE::RenderGrid::Private::FRenderGridGenericExecutionQueueActionReturningDelay::CreateLambda([this, Job]() -> UE::RenderGrid::Private::FRenderGridGenericExecutionQueueDelay
 				{
-					PreviousProps = UE::RenderGrid::IRenderGridModule::Get().GetManager().ApplyJobPropValues(RenderGrid, Job);
-					return UE::RenderGrid::Private::FRenderGridGenericExecutionQueueDelay::Frames(2 + Job->GetWaitFramesBeforeRendering());
+					if (IsValid(RenderGrid))
+					{
+						PreviousProps = UE::RenderGrid::IRenderGridModule::Get().GetManager().ApplyJobPropValues(RenderGrid, Job);
+					}
+					return UE::RenderGrid::Private::FRenderGridGenericExecutionQueueDelay::Frames(2 + (IsValid(RenderGrid) ? Job->GetWaitFramesBeforeRendering() : 0));
 				}));
 
 				Queue->Add(UE::RenderGrid::Private::FRenderGridGenericExecutionQueueAction::CreateLambda([this, Job]()
 				{
-					RenderGrid->BeginJobRender(this, Job);
+					if (IsValid(RenderGrid) && IsValid(Job))
+					{
+						RenderGrid->BeginJobRender(this, Job);
+					}
 				}));
 
 				Queue->Add(UE::RenderGrid::Private::FRenderGridGenericExecutionQueueActionReturningDelayFuture::CreateLambda([Entry]() -> TSharedFuture<void>
@@ -543,7 +645,10 @@ void URenderGridQueue::OnProcessJob(URenderGridJob* Job)
 
 				Queue->Add(UE::RenderGrid::Private::FRenderGridGenericExecutionQueueAction::CreateLambda([this, Job]()
 				{
-					RenderGrid->EndJobRender(this, Job);
+					if (IsValid(RenderGrid) && IsValid(Job))
+					{
+						RenderGrid->EndJobRender(this, Job);
+					}
 				}));
 
 				Queue->Add(UE::RenderGrid::Private::FRenderGridGenericExecutionQueueActionReturningDelay::CreateLambda([this]() -> UE::RenderGrid::Private::FRenderGridGenericExecutionQueueDelay
@@ -566,7 +671,7 @@ void URenderGridQueue::OnProcessJob(URenderGridJob* Job)
 
 	Queue->Add(UE::RenderGrid::Private::FRenderGridGenericExecutionQueueAction::CreateLambda([this]()
 	{
-		while (!RemainingJobs.IsEmpty() && !bCanceled)
+		while (!RemainingJobs.IsEmpty() && !bCanceled && !bFinished)
 		{
 			TObjectPtr<URenderGridJob> NextJob = RemainingJobs[0];
 			RemainingJobs.RemoveAt(0, 1, false);
@@ -577,16 +682,36 @@ void URenderGridQueue::OnProcessJob(URenderGridJob* Job)
 				return;
 			}
 		}
-		OnFinish();
+
+		Queue->DelayFrames(2);
+		Queue->Add(UE::RenderGrid::Private::FRenderGridGenericExecutionQueueAction::CreateLambda([this]()
+		{
+			OnFinish();
+		}));
 	}));
 }
 
 void URenderGridQueue::OnFinish()
 {
-	bool bSuccess = !bCanceled;
-	Cancel();// to prevent any new jobs from being added to it
+	if (bFinished)
+	{
+		return;
+	}
+	bFinished = true; // prevents any new jobs from being added to it
 
-	if (Args.bIsBatchRender)
+	if (!bStarted)
+	{
+		return;
+	}
+
+	if (GetCurrentlyExecutingQueue() == this)
+	{
+		DoNextExecutingQueue(); // will start the next queue (during the next tick)
+	}
+
+	Queue->Stop();
+
+	if (Args.bIsBatchRender && IsValid(RenderGrid))
 	{
 		RenderGrid->EndBatchRender(this);
 	}
@@ -595,5 +720,7 @@ void URenderGridQueue::OnFinish()
 	PreviousFrameLimitSettings = FRenderGridPreviousEngineFpsSettings();
 
 	RemoveFromRoot();
-	OnExecuteFinishedDelegate.Broadcast(this, bSuccess);
+
+	OnExecuteFinishedDelegate.Broadcast(this, !bCanceled);
+	OnExecuteFinishedDelegate.Clear();
 }
diff --git a/Engine/Plugins/Experimental/RenderGrid/Source/RenderGrid/Private/Utils/RenderGridGenericExecutionQueue.cpp b/Engine/Plugins/Experimental/RenderGrid/Source/RenderGrid/Private/Utils/RenderGridGenericExecutionQueue.cpp
--- a/Engine/Plugins/Experimental/RenderGrid/Source/RenderGrid/Private/Utils/RenderGridGenericExecutionQueue.cpp
+++ b/Engine/Plugins/Experimental/RenderGrid/Source/RenderGrid/Private/Utils/RenderGridGenericExecutionQueue.cpp
@@ -40,12 +40,12 @@ void UE::RenderGrid::Private::FRenderGridGenericExecutionQueue::Tick(float Delta
 
 void UE::RenderGrid::Private::FRenderGridGenericExecutionQueue::Start()
 {
-	bStarted = true;
+	bRunning = true;
 }
 
 void UE::RenderGrid::Private::FRenderGridGenericExecutionQueue::Stop()
 {
-	bStarted = false;
+	bRunning = false;
 }
 
 void UE::RenderGrid::Private::FRenderGridGenericExecutionQueue::ExecuteNext()
@@ -63,7 +63,7 @@ void UE::RenderGrid::Private::FRenderGridGenericExecutionQueue::ExecuteNext()
 			// continue in Tick event
 			return;
 		}
-		if (bStarted && ExecuteNextEntry())
+		if (bRunning && ExecuteNextEntry())
 		{
 			// execution completed, restart this loop  (to check for new delays, and to execute the next entry)
 			continue;
diff --git a/Engine/Plugins/Experimental/RenderGrid/Source/RenderGrid/Private/Utils/RenderGridGenericExecutionQueue.h b/Engine/Plugins/Experimental/RenderGrid/Source/RenderGrid/Private/Utils/RenderGridGenericExecutionQueue.h
--- a/Engine/Plugins/Experimental/RenderGrid/Source/RenderGrid/Private/Utils/RenderGridGenericExecutionQueue.h
+++ b/Engine/Plugins/Experimental/RenderGrid/Source/RenderGrid/Private/Utils/RenderGridGenericExecutionQueue.h
@@ -131,7 +131,7 @@ namespace UE::RenderGrid::Private
 		void Stop();
 
 		/** Returns true if Start() has been called and Stop() hasn't been called yet. */
-		bool IsRunning() { return bStarted; }
+		bool IsRunning() { return bRunning; }
 
 	protected:
 		/** Executes the next delay (if there are any), otherwise it executes the next entry (action). */
@@ -155,7 +155,7 @@ namespace UE::RenderGrid::Private
 
 	protected:
 		/** Whether it has started (and hasn't been stopped/paused yet). This means that if this is true, Start() has been called, and Stop() hasn't been called yet since then. */
-		bool bStarted = false;
+		bool bRunning = false;
 
 		/** Whether it's currently executing a delay or an entry (action). False means there were no delays and entries queued up anymore. */
 		bool bExecuting = false;
diff --git a/Engine/Plugins/Experimental/RenderGrid/Source/RenderGrid/Public/RenderGrid/RenderGrid.h b/Engine/Plugins/Experimental/RenderGrid/Source/RenderGrid/Public/RenderGrid/RenderGrid.h
--- a/Engine/Plugins/Experimental/RenderGrid/Source/RenderGrid/Public/RenderGrid/RenderGrid.h
+++ b/Engine/Plugins/Experimental/RenderGrid/Source/RenderGrid/Public/RenderGrid/RenderGrid.h
@@ -38,6 +38,43 @@ public:
 };
 
 
+/**
+ * This class contains the default values of render grid jobs.
+ *
+ * This is placed in a separate class, primarily so that it be separated easily from the render grid when using the details view widget.
+ */
+UCLASS()
+class RENDERGRID_API URenderGridSettings : public UObject
+{
+	GENERATED_BODY()
+
+public:
+	URenderGridSettings();
+
+public:
+	/** The type of the properties that a job in this grid can have. */
+	UPROPERTY(/*EditInstanceOnly, Category="Render Grid", Meta=(DisplayName="Properties Type")*/)
+	ERenderGridPropsSourceType PropsSourceType;
+
+	/** The remote control properties that a job in this grid can have, only use this if PropsSourceType is ERenderGridPropsSourceType::RemoteControl. */
+	UPROPERTY(EditInstanceOnly, Category="Render Grid", Meta=(DisplayName="Remote Control Preset" /*, EditCondition="PropsSourceType == ERenderGridPropsSourceType::RemoteControl", EditConditionHides*/))
+	TObjectPtr<URemoteControlPreset> PropsSourceOrigin_RemoteControl;
+
+public:
+	/** GetPropsSource calls are somewhat expensive, we speed that up by caching the result (the PropsSource) that has been last outputted by that function. */
+	UPROPERTY(Transient)
+	mutable TObjectPtr<URenderGridPropsSourceBase> CachedPropsSource;
+
+	/** GetPropsSource calls are somewhat expensive, we speed that up by caching the PropsSourceType last used in that function. */
+	UPROPERTY(Transient)
+	mutable ERenderGridPropsSourceType CachedPropsSourceType;
+
+	/** GetPropsSource calls are somewhat expensive, we speed that up by caching the PropsSourceOrigin last used in that function. */
+	UPROPERTY(Transient)
+	mutable TWeakObjectPtr<UObject> CachedPropsSourceOriginWeakPtr;
+};
+
+
 /**
  * This class contains the default values of render grid jobs.
  *
@@ -444,6 +481,19 @@ public:
 	/** Overridable native event for when job rendering for the viewport viewer-mode ends. */
 	virtual void EndViewportRender(URenderGridJob* Job);
 
+public:
+	/** Renders all the enabled jobs of this render grid. */
+	UFUNCTION(BlueprintCallable, Category="Render Grid")
+	URenderGridQueue* Render();
+
+	/** Renders all the given jobs of this render grid. */
+	UFUNCTION(BlueprintCallable, Category="Render Grid")
+	URenderGridQueue* RenderJobs(const TArray<URenderGridJob*>& Jobs);
+
+	/** Renders the given job of this render grid. */
+	UFUNCTION(BlueprintCallable, Category="Render Grid")
+	URenderGridQueue* RenderJob(URenderGridJob* Job);
+
 public:
 	/** Returns the GUID, which is randomly generated at creation. */
 	UFUNCTION(BlueprintPure, Category="Render Grid")
@@ -453,6 +503,9 @@ public:
 	UFUNCTION(BlueprintCallable, Category="Render Grid")
 	void GenerateNewGuid() { Guid = FGuid::NewGuid(); }
 
+public:
+	URenderGridSettings* GetSettingsObject() { return Settings; }
+
 	UFUNCTION(BlueprintCallable, Category="Render Grid")
 	void SetPropsSource(ERenderGridPropsSourceType InPropsSourceType, UObject* InPropsSourceOrigin = nullptr);
 
@@ -467,7 +520,7 @@ public:
 	}
 
 	UFUNCTION(BlueprintPure, Category="Render Grid")
-	ERenderGridPropsSourceType GetPropsSourceType() const { return PropsSourceType; }
+	ERenderGridPropsSourceType GetPropsSourceType() const { return Settings->PropsSourceType; }
 
 	UFUNCTION(BlueprintPure, Category="Render Grid")
 	UObject* GetPropsSourceOrigin() const;
@@ -504,7 +557,7 @@ public:
 public:
 	UFUNCTION(BlueprintCallable, Category="Render Grid", Meta=(Keywords="remove delete all"))
 	void ClearRenderGridJobs();
-	
+
 	UFUNCTION(BlueprintCallable, Category="Render Grid")
 	void SetRenderGridJobs(const TArray<URenderGridJob*>& Jobs);
 
@@ -582,21 +635,15 @@ private:
 	FGuid Guid;
 
 
-	/** The type of the properties that a job in this grid can have. */
-	UPROPERTY(/*EditInstanceOnly, Category="Render Grid", Meta=(DisplayName="Properties Type", AllowPrivateAccess="true")*/)
-	ERenderGridPropsSourceType PropsSourceType;
-
-	/** The remote control properties that a job in this grid can have, only use this if PropsSourceType is ERenderGridPropsSourceType::RemoteControl. */
-	UPROPERTY(EditInstanceOnly, Category="Render Grid", Meta=(DisplayName="Remote Control Preset", AllowPrivateAccess="true" /*, EditCondition="PropsSourceType == ERenderGridPropsSourceType::RemoteControl", EditConditionHides*/))
-	TObjectPtr<URemoteControlPreset> PropsSourceOrigin_RemoteControl;
-
+	/** The settings of this render grid. */
+	UPROPERTY(Instanced)
+	TObjectPtr<URenderGridSettings> Settings;
 
 	/** The default values for new jobs. */
 	UPROPERTY(Instanced)
 	TObjectPtr<URenderGridDefaults> Defaults;
 
-
-	/** The jobs of this grid. */
+	/** The jobs of this render grid. */
 	UPROPERTY(Instanced)
 	TArray<TObjectPtr<URenderGridJob>> RenderGridJobs;
 
@@ -606,19 +653,6 @@ private:
 	bool bExecutingBlueprintEvent;
 
 
-	/** GetPropsSource calls are somewhat expensive, we speed that up by caching the result (the PropsSource) that has been last outputted by that function. */
-	UPROPERTY(Transient)
-	mutable TObjectPtr<URenderGridPropsSourceBase> CachedPropsSource;
-
-	/** GetPropsSource calls are somewhat expensive, we speed that up by caching the PropsSourceType last used in that function. */
-	UPROPERTY(Transient)
-	mutable ERenderGridPropsSourceType CachedPropsSourceType;
-
-	/** GetPropsSource calls are somewhat expensive, we speed that up by caching the PropsSourceOrigin last used in that function. */
-	UPROPERTY(Transient)
-	mutable TWeakObjectPtr<UObject> CachedPropsSourceOriginWeakPtr;
-
-
 	/** GetWorld calls can be expensive, we speed them up by caching the last found world until it goes away. */
 	mutable TWeakObjectPtr<UWorld> CachedWorldWeakPtr;
 };
diff --git a/Engine/Plugins/Experimental/RenderGrid/Source/RenderGrid/Public/RenderGrid/RenderGridManager.h b/Engine/Plugins/Experimental/RenderGrid/Source/RenderGrid/Public/RenderGrid/RenderGridManager.h
--- a/Engine/Plugins/Experimental/RenderGrid/Source/RenderGrid/Public/RenderGrid/RenderGridManager.h
+++ b/Engine/Plugins/Experimental/RenderGrid/Source/RenderGrid/Public/RenderGrid/RenderGridManager.h
@@ -87,11 +87,11 @@ namespace UE::RenderGrid
 
 
 	public:
-		/** Batch render the currently enabled render grid job(s) of the given render grid. **/
-		URenderGridQueue* CreateBatchRenderQueue(URenderGrid* Grid);
+		/** Batch render the given render grid job(s) of the given render grid. */
+		URenderGridQueue* CreateBatchRenderQueue(URenderGrid* Grid, const TArray<URenderGridJob*>& Jobs);
 
 
-		/** Render a preview frame (or multiple if no frame number is specified) of the given render grid job. **/
+		/** Render a preview frame (or multiple if no frame number is specified) of the given render grid job. */
 		URenderGridQueue* RenderPreviewFrame(const FRenderGridManagerRenderPreviewFrameArgs& Args);
 
 		/** Gets the rendered preview frame (of a rendering in which the frame number was specified). */
diff --git a/Engine/Plugins/Experimental/RenderGrid/Source/RenderGrid/Public/RenderGrid/RenderGridQueue.h b/Engine/Plugins/Experimental/RenderGrid/Source/RenderGrid/Public/RenderGrid/RenderGridQueue.h
--- a/Engine/Plugins/Experimental/RenderGrid/Source/RenderGrid/Public/RenderGrid/RenderGridQueue.h
+++ b/Engine/Plugins/Experimental/RenderGrid/Source/RenderGrid/Public/RenderGrid/RenderGridQueue.h
@@ -4,6 +4,9 @@
 
 #include "RenderGridManager.h"
 #include "RenderGridUtils.h"
+#include "Tickable.h"
+#include "Containers/Queue.h"
+#include "UObject/ObjectPtr.h"
 #include "RenderGridQueue.generated.h"
 
 
@@ -154,11 +157,49 @@ protected:
  * This class is responsible for rendering the given render grid jobs.
  */
 UCLASS(BlueprintType)
-class RENDERGRID_API URenderGridQueue : public UObject
+class RENDERGRID_API URenderGridQueue : public UObject, public FTickableGameObject
 {
 	GENERATED_BODY()
 
+
 public:
+	static bool IsExecutingAny();
+
+	DECLARE_MULTICAST_DELEGATE(FOnExecutionQueueChanged);
+	static FOnExecutionQueueChanged& OnExecutionQueueChanged() { return OnExecutionQueueChangedDelegate; }
+
+private:
+	/** Returns the currently executing queue, or NULL if there isn't any currently executing. */
+	static URenderGridQueue* GetCurrentlyExecutingQueue();
+
+	/** Adds the given queue. Will automatically execute it if it's the only one currently. */
+	static void AddExecutingQueue(URenderGridQueue* Queue);
+
+	/** Executes the next queue, if any. Nothing will happen until the next tick. */
+	static void DoNextExecutingQueue();
+
+private:
+	/** The queue of executing queues. This contains the currently rendering queue, as well as the ones that should be rendering after it. */
+	static TQueue<TObjectPtr<URenderGridQueue>> ExecutingQueues;
+
+	/** The delegate for when data in the ExecutingQueues has changed. */
+	static FOnExecutionQueueChanged OnExecutionQueueChangedDelegate;
+
+
+public:
+	//~ Begin FTickableGameObject interface
+	virtual void Tick(float DeltaTime) override;
+	virtual ETickableTickType GetTickableTickType() const override { return ETickableTickType::Always; }
+	virtual bool IsTickableWhenPaused() const override { return true; }
+	virtual bool IsTickableInEditor() const override { return true; }
+	virtual bool IsTickable() const override { return true; }
+	virtual bool IsAllowedToTick() const override { return true; }
+	virtual TStatId GetStatId() const override
+	{
+		RETURN_QUICK_DECLARE_CYCLE_STAT(FRenderGridQueue, STATGROUP_Tickables);
+	}
+	//~ End FTickableGameObject interface
+
 	/** Creates a new render queue, it won't be started right away. */
 	static URenderGridQueue* Create(const UE::RenderGrid::FRenderGridQueueCreateArgs& Args);
 
@@ -182,36 +223,52 @@ public:
 	UFUNCTION(BlueprintCallable, Category="Render Grid Queue", Meta=(Keywords="stop quit exit kill terminate end"))
 	void Cancel();
 
+	/** Returns true if this queue has been started. */
+	UFUNCTION(BlueprintPure, Category="Render Grid Queue", Meta=(Keywords="stopped quited exited killed terminated ended succeeded completed finished canceled cancelled"))
+	bool IsStarted() const { return bStarted; }
+
+	/** Returns true if this queue is currently paused. */
+	UFUNCTION(BlueprintPure, Category="Render Grid Queue", Meta=(Keywords="stopped quited exited killed terminated ended succeeded completed finished canceled cancelled"))
+	bool IsPaused() const { return (bPaused && !bCanceled && !bFinished); }
+
 	/** Returns true if this queue has been canceled. */
-	UFUNCTION(BlueprintCallable, Category="Render Grid Queue", Meta=(Keywords="stopped quited exited killed terminated ended"))
+	UFUNCTION(BlueprintPure, Category="Render Grid Queue", Meta=(Keywords="stopped quited exited killed terminated ended succeeded completed finished canceled cancelled"))
 	bool IsCanceled() const { return bCanceled; }
 
+	/** Returns true if this queue has been canceled. */
+	UFUNCTION(BlueprintPure, Category="Render Grid Queue", Meta=(Keywords="stopped quited exited killed terminated ended succeeded completed finished canceled cancelled"))
+	bool IsFinished() const { return bFinished; }
+
+	/** Returns true if this queue is the one that's currently rendering, returns false if it hasn't started yet, or if it's waiting in the queue, or if it has finished. */
+	UFUNCTION(BlueprintPure, Category="Render Grid Queue", Meta=(Keywords="stopped quited exited killed terminated ended succeeded completed finished canceled cancelled"))
+	bool IsCurrentlyRendering() const { return (GetCurrentlyExecutingQueue() == this); }
+
 	/** Retrieves the rendering status of the given render grid job. */
-	UFUNCTION(BlueprintCallable, Category="Render Grid Queue", Meta=(Keywords="rendering progression obtain text"))
+	UFUNCTION(BlueprintPure, Category="Render Grid Queue", Meta=(Keywords="rendering progression obtain text"))
 	FString GetJobStatus(URenderGridJob* Job) const;
 
 	/** Returns all the jobs that have been and will be rendered. */
-	UFUNCTION(BlueprintCallable, Category="Render Grid Queue", Meta=(Keywords="rendering progression obtain"))
+	UFUNCTION(BlueprintPure, Category="Render Grid Queue", Meta=(Keywords="rendering progression obtain"))
 	TArray<URenderGridJob*> GetJobs() const;
 
 	/** Returns the number of jobs that have been and will be rendered. */
-	UFUNCTION(BlueprintCallable, Category="Render Grid Queue", Meta=(Keywords="rendering progression obtain number amount"))
+	UFUNCTION(BlueprintPure, Category="Render Grid Queue", Meta=(Keywords="rendering progression obtain number amount"))
 	int32 GetJobsCount() const;
 
 	/** Returns the number of jobs that have finished rendering. Basically just returns [Get Jobs Count] minus [Get Jobs Remaining Count]. */
-	UFUNCTION(BlueprintCallable, Category="Render Grid Queue", Meta=(Keywords="rendering progression obtain number amount finished"))
+	UFUNCTION(BlueprintPure, Category="Render Grid Queue", Meta=(Keywords="rendering progression obtain number amount finished"))
 	int32 GetJobsCompletedCount() const;
 
 	/** Returns the percentage of jobs finished, this includes the progression of the job that is currently rendering. */
-	UFUNCTION(BlueprintCallable, Category="Render Grid Queue", Meta=(Keywords="rendering progression obtain number amount finished"))
+	UFUNCTION(BlueprintPure, Category="Render Grid Queue", Meta=(Keywords="rendering progression obtain number amount finished"))
 	float GetStatusPercentage() const;
 
 	/** Returns the number of jobs that are still left to render, includes the job that is currently rendering. */
-	UFUNCTION(BlueprintCallable, Category="Render Grid Queue", Meta=(Keywords="rendering progression obtain number amount finished"))
+	UFUNCTION(BlueprintPure, Category="Render Grid Queue", Meta=(Keywords="rendering progression obtain number amount finished"))
 	int32 GetJobsRemainingCount() const;
 
 	/** Returns the status of the rendering process. */
-	UFUNCTION(BlueprintCallable, Category="Render Grid Queue", Meta=(Keywords="rendering progression obtain text"))
+	UFUNCTION(BlueprintPure, Category="Render Grid Queue", Meta=(Keywords="rendering progression obtain text"))
 	FString GetStatus() const;
 
 protected:
@@ -238,10 +295,22 @@ protected:
 	UPROPERTY(Transient)
 	TObjectPtr<URenderGrid> RenderGrid;
 
-	/** Whether the remaining render grid jobs should be prevented from rendering. */
+	/** Whether the queue has been started yet. */
+	UPROPERTY(Transient)
+	bool bStarted;
+
+	/** Whether the queue has been paused. */
+	UPROPERTY(Transient)
+	bool bPaused;
+
+	/** Whether the queue has been canceled. */
 	UPROPERTY(Transient)
 	bool bCanceled;
 
+	/** Whether the queue has been finished. */
+	UPROPERTY(Transient)
+	bool bFinished;
+
 	/** The property values that have been overwritten by the currently applied render grid job property values. */
 	UPROPERTY(Transient)
 	FRenderGridManagerPreviousPropValues PreviousProps;
diff --git a/Engine/Plugins/Experimental/RenderGrid/Source/RenderGridDeveloper/Private/Blueprints/RenderGridBlueprint.cpp b/Engine/Plugins/Experimental/RenderGrid/Source/RenderGridDeveloper/Private/Blueprints/RenderGridBlueprint.cpp
--- a/Engine/Plugins/Experimental/RenderGrid/Source/RenderGridDeveloper/Private/Blueprints/RenderGridBlueprint.cpp
+++ b/Engine/Plugins/Experimental/RenderGrid/Source/RenderGridDeveloper/Private/Blueprints/RenderGridBlueprint.cpp
@@ -66,7 +66,7 @@ void URenderGridBlueprint::RunOnInstances(const UE::RenderGrid::FRenderGridBluep
 {
 	if (UClass* MyRenderGridClass = GeneratedClass; IsValid(MyRenderGridClass))
 	{
-		if (URenderGrid* DefaultObject = Cast<URenderGrid>(MyRenderGridClass->GetDefaultObject(false)); IsValid(DefaultObject))
+		if (URenderGrid* DefaultObject = Cast<URenderGrid>(MyRenderGridClass->GetDefaultObject(true)); IsValid(DefaultObject))
 		{
 			Callback.ExecuteIfBound(DefaultObject);
 
@@ -95,7 +95,7 @@ void URenderGridBlueprint::Save()
 {
 	if (UClass* MyRenderGridClass = GeneratedClass; IsValid(MyRenderGridClass))
 	{
-		if (URenderGrid* DefaultObject = Cast<URenderGrid>(MyRenderGridClass->GetDefaultObject(false)); IsValid(DefaultObject))
+		if (URenderGrid* DefaultObject = Cast<URenderGrid>(MyRenderGridClass->GetDefaultObject(true)); IsValid(DefaultObject))
 		{
 			for (URenderGridJob* Job : DefaultObject->GetRenderGridJobsRef())
 			{
@@ -147,6 +147,31 @@ void URenderGridBlueprint::PropagateAllPropertiesToAsset(URenderGrid* Instance)
 	RenderGrid->CopyAllProperties(Instance);
 }
 
+URenderGrid* URenderGridBlueprint::GetRenderGridWithBlueprintGraph() const
+{
+	URenderGrid* Result = GetRenderGrid();
+	const_cast<URenderGridBlueprint*>(this)->RunOnInstances(UE::RenderGrid::FRenderGridBlueprintRunOnInstancesCallback::CreateLambda([&Result](URenderGrid* Instance)
+	{
+		if (!IsValid(Result) || Result->HasAnyFlags(RF_ClassDefaultObject | RF_DefaultSubObject))
+		{
+			Result = Instance;
+		}
+	}));
+	return Result;
+}
+
+URenderGrid* URenderGridBlueprint::GetRenderGridClassDefaultObject() const
+{
+	if (UClass* MyRenderGridClass = GeneratedClass; IsValid(MyRenderGridClass))
+	{
+		if (URenderGrid* DefaultObject = Cast<URenderGrid>(MyRenderGridClass->GetDefaultObject(true)); IsValid(DefaultObject))
+		{
+			return DefaultObject;
+		}
+	}
+	return nullptr;
+}
+
 void URenderGridBlueprint::OnPostVariablesChange(UBlueprint* InBlueprint)
 {
 	if (InBlueprint != this)
diff --git a/Engine/Plugins/Experimental/RenderGrid/Source/RenderGridDeveloper/Private/RenderGridDeveloperLibrary.cpp b/Engine/Plugins/Experimental/RenderGrid/Source/RenderGridDeveloper/Private/RenderGridDeveloperLibrary.cpp
new file mode 100644
--- /dev/null
+++ b/Engine/Plugins/Experimental/RenderGrid/Source/RenderGridDeveloper/Private/RenderGridDeveloperLibrary.cpp
@@ -0,0 +1,44 @@
+﻿// Copyright Epic Games, Inc. All Rights Reserved.
+
+#include "RenderGridDeveloperLibrary.h"
+#include "AssetRegistry/AssetData.h"
+#include "AssetRegistry/AssetRegistryModule.h"
+#include "AssetRegistry/IAssetRegistry.h"
+#include "Blueprints/RenderGridBlueprint.h"
+#include "RenderGrid/RenderGrid.h"
+
+
+TArray<URenderGridBlueprint*> URenderGridDeveloperLibrary::GetAllRenderGridBlueprintAssets()
+{
+	FARFilter Filter;
+	Filter.bIncludeOnlyOnDiskAssets = false;
+	Filter.ClassPaths = {URenderGridBlueprint::StaticClass()->GetClassPathName()};
+	Filter.bRecursivePaths = true;
+
+	TArray<FAssetData> AssetsData;
+	IAssetRegistry& AssetRegistry = FModuleManager::Get().LoadModuleChecked<FAssetRegistryModule>(TEXT("AssetRegistry")).Get();
+	AssetRegistry.GetAssets(Filter, AssetsData);
+
+	TArray<URenderGridBlueprint*> Result;
+	for (const FAssetData& AssetData : AssetsData)
+	{
+		if (URenderGridBlueprint* RenderGridBlueprint = Cast<URenderGridBlueprint>(AssetData.GetAsset()); IsValid(RenderGridBlueprint))
+		{
+			Result.Add(RenderGridBlueprint);
+		}
+	}
+	return Result;
+}
+
+TArray<URenderGrid*> URenderGridDeveloperLibrary::GetAllRenderGridAssets()
+{
+	TArray<URenderGrid*> Result;
+	for (URenderGridBlueprint* RenderGridBlueprint : GetAllRenderGridBlueprintAssets())
+	{
+		if (URenderGrid* RenderGrid = RenderGridBlueprint->GetRenderGridWithBlueprintGraph(); IsValid(RenderGrid))
+		{
+			Result.Add(RenderGrid);
+		}
+	}
+	return Result;
+}
diff --git a/Engine/Plugins/Experimental/RenderGrid/Source/RenderGridDeveloper/Public/Blueprints/RenderGridBlueprint.h b/Engine/Plugins/Experimental/RenderGrid/Source/RenderGridDeveloper/Public/Blueprints/RenderGridBlueprint.h
--- a/Engine/Plugins/Experimental/RenderGrid/Source/RenderGridDeveloper/Public/Blueprints/RenderGridBlueprint.h
+++ b/Engine/Plugins/Experimental/RenderGrid/Source/RenderGridDeveloper/Public/Blueprints/RenderGridBlueprint.h
@@ -59,10 +59,18 @@ public:
 	void PropagateAllPropertiesToAsset(URenderGrid* Instance);
 
 public:
-	/** Returns the RenderGrid reference that this RenderGrid asset contains. */
+	/** Returns the RenderGrid reference that this RenderGrid asset contains. This is simply the data representation of the render grid, meaning that it won't contain a blueprint graph or any user code. */
 	UFUNCTION(BlueprintPure, Category="Render Grid")
 	URenderGrid* GetRenderGrid() const { return RenderGrid; }
 
+	/** Returns the RenderGrid reference that this RenderGrid asset contains. This will be the subclass of the blueprint class, meaning it will contain a blueprint graph. */
+	UFUNCTION(BlueprintPure, Category="Render Grid")
+	URenderGrid* GetRenderGridWithBlueprintGraph() const;
+
+	/** Returns the RenderGrid reference that this RenderGrid asset contains. This will be the default object of the subclass of the blueprint class, meaning it will contain a blueprint graph. */
+	UFUNCTION(BlueprintPure, Category="Render Grid")
+	URenderGrid* GetRenderGridClassDefaultObject() const;
+
 private:
 	virtual void OnPostVariablesChange(UBlueprint* InBlueprint);
 
diff --git a/Engine/Plugins/Experimental/RenderGrid/Source/RenderGridDeveloper/Public/RenderGridDeveloperLibrary.h b/Engine/Plugins/Experimental/RenderGrid/Source/RenderGridDeveloper/Public/RenderGridDeveloperLibrary.h
new file mode 100644
--- /dev/null
+++ b/Engine/Plugins/Experimental/RenderGrid/Source/RenderGridDeveloper/Public/RenderGridDeveloperLibrary.h
@@ -0,0 +1,36 @@
+﻿// Copyright Epic Games, Inc. All Rights Reserved.
+
+#pragma once
+
+#include "CoreMinimal.h"
+#include "Kismet/BlueprintFunctionLibrary.h"
+#include "RenderGridDeveloperLibrary.generated.h"
+
+
+class URenderGrid;
+class URenderGridBlueprint;
+
+
+UCLASS(meta=(ScriptName="RenderGridDeveloperLibrary"))
+class RENDERGRIDDEVELOPER_API URenderGridDeveloperLibrary : public UBlueprintFunctionLibrary
+{
+	GENERATED_BODY()
+
+public:
+	/**
+	 * Returns all render grid assets that currently exist in the project (on disk and in memory).
+	 * Will load the render grid assets in that are currently unloaded.
+	 * 
+	 * This is a slow operation, so avoid doing this every tick.
+	 */
+	static TArray<URenderGridBlueprint*> GetAllRenderGridBlueprintAssets();
+
+	/**
+	 * Returns all render grid assets that currently exist in the project (on disk and in memory).
+	 * Will load the render grid assets in that are currently unloaded.
+	 * 
+	 * This is a slow operation, so avoid doing this every tick.
+	 */
+	UFUNCTION(BlueprintCallable, Category="Render Grid")
+	static TArray<URenderGrid*> GetAllRenderGridAssets();
+};
diff --git a/Engine/Plugins/Experimental/RenderGrid/Source/RenderGridDeveloper/RenderGridDeveloper.Build.cs b/Engine/Plugins/Experimental/RenderGrid/Source/RenderGridDeveloper/RenderGridDeveloper.Build.cs
--- a/Engine/Plugins/Experimental/RenderGrid/Source/RenderGridDeveloper/RenderGridDeveloper.Build.cs
+++ b/Engine/Plugins/Experimental/RenderGrid/Source/RenderGridDeveloper/RenderGridDeveloper.Build.cs
@@ -15,6 +15,7 @@ public class RenderGridDeveloper : ModuleRules
 			{
 				"Core",
 				"CoreUObject",
+				"Engine",
 				"KismetCompiler",
 			}
 		);
diff --git a/Engine/Plugins/Experimental/RenderGrid/Source/RenderGridEditor/Private/Toolkit/RenderGridEditor.cpp b/Engine/Plugins/Experimental/RenderGrid/Source/RenderGridEditor/Private/Toolkit/RenderGridEditor.cpp
--- a/Engine/Plugins/Experimental/RenderGrid/Source/RenderGridEditor/Private/Toolkit/RenderGridEditor.cpp
+++ b/Engine/Plugins/Experimental/RenderGrid/Source/RenderGridEditor/Private/Toolkit/RenderGridEditor.cpp
@@ -7,7 +7,6 @@
 #include "RenderGrid/RenderGridQueue.h"
 #include "RenderGridEditorToolbar.h"
 #include "IRenderGridEditorModule.h"
-#include "IRenderGridModule.h"
 
 #include "BlueprintCompilationManager.h"
 #include "BlueprintModes/RenderGridApplicationModeBase.h"
@@ -21,7 +20,6 @@
 #include "Framework/Commands/GenericCommands.h"
 #include "Framework/MultiBox/MultiBoxBuilder.h"
 #include "Kismet2/BlueprintEditorUtils.h"
-#include "Misc/MessageDialog.h"
 #include "PropertyCustomizationHelpers.h"
 #include "SBlueprintEditorToolbar.h"
 #include "ScopedTransaction.h"
@@ -74,6 +72,8 @@ bool URenderGridJobSelection::SetSelectedRenderGridJobs(const TArray<URenderGrid
 }
 
 
+TStrongObjectPtr<URenderGridJobSelection> UE::RenderGrid::Private::FRenderGridEditor::RenderGridJobSelection = TStrongObjectPtr(NewObject<URenderGridJobSelection>());
+
 void UE::RenderGrid::Private::FRenderGridEditor::InitRenderGridEditor(const EToolkitMode::Type Mode, const TSharedPtr<IToolkitHost>& InitToolkitHost, URenderGridBlueprint* InRenderGridBlueprint)
 {
 	check(InRenderGridBlueprint);
@@ -125,12 +125,12 @@ void UE::RenderGrid::Private::FRenderGridEditor::InitRenderGridEditor(const EToo
 
 UE::RenderGrid::Private::FRenderGridEditor::FRenderGridEditor()
 	: PreviewBlueprintWeakPtr(nullptr)
-	, RenderGridJobSelection(NewObject<URenderGridJobSelection>())
 	, bRunRenderNewBatch(false)
 	, BatchRenderQueue(nullptr)
 	, PreviewRenderQueue(nullptr)
 	, DebuggingTimeInSecondsRemaining(0)
 	, bIsDebugging(false)
+	, bPreviousShouldHideUI(false)
 {}
 
 UE::RenderGrid::Private::FRenderGridEditor::~FRenderGridEditor()
@@ -279,12 +279,20 @@ void UE::RenderGrid::Private::FRenderGridEditor::SetIsDebugging(const bool bInIs
 
 bool UE::RenderGrid::Private::FRenderGridEditor::IsBatchRendering() const
 {
-	return IsValid(BatchRenderQueue.Get());
+	if (URenderGridQueue* Queue = BatchRenderQueue.Get(); IsValid(Queue))
+	{
+		return Queue->IsCurrentlyRendering();
+	}
+	return false;
 }
 
 bool UE::RenderGrid::Private::FRenderGridEditor::IsPreviewRendering() const
 {
-	return IsValid(PreviewRenderQueue.Get());
+	if (URenderGridQueue* Queue = PreviewRenderQueue.Get(); IsValid(Queue))
+	{
+		return Queue->IsCurrentlyRendering();
+	}
+	return false;
 }
 
 void UE::RenderGrid::Private::FRenderGridEditor::SetPreviewRenderQueue(URenderGridQueue* Queue)
@@ -305,6 +313,11 @@ void UE::RenderGrid::Private::FRenderGridEditor::MarkAsModified()
 TArray<URenderGridJob*> UE::RenderGrid::Private::FRenderGridEditor::GetSelectedRenderGridJobs() const
 {
 	TArray<URenderGridJob*> SelectedJobs;
+	if (ShouldHideUI())
+	{
+		return SelectedJobs;
+	}
+
 	if (URenderGrid* Grid = GetInstance(); IsValid(Grid))
 	{
 		if (URenderGridJobSelection* JobSelection = RenderGridJobSelection.Get(); IsValid(JobSelection))
@@ -325,15 +338,8 @@ void UE::RenderGrid::Private::FRenderGridEditor::SetSelectedRenderGridJobs(const
 {
 	if (URenderGridJobSelection* JobSelection = RenderGridJobSelection.Get(); IsValid(JobSelection))
 	{
-		{
-			JobSelection->Modify();
-			if (!JobSelection->SetSelectedRenderGridJobs(Jobs))
-			{
-				return;
-			}
-		}
-
-		OnRenderGridJobsSelectionChanged().Broadcast();
+		JobSelection->Modify();
+		JobSelection->SetSelectedRenderGridJobs(Jobs);
 	}
 }
 
@@ -389,6 +395,29 @@ void UE::RenderGrid::Private::FRenderGridEditor::Tick(float DeltaTime)
 		BatchRenderListAction();
 	}
 
+	if (ShouldHideUI() != bPreviousShouldHideUI)
+	{
+		bPreviousShouldHideUI = !bPreviousShouldHideUI;
+		OnRenderGridShouldHideUIChanged().Broadcast();
+	}
+
+	if (URenderGridJobSelection* JobSelection = RenderGridJobSelection.Get(); IsValid(JobSelection))
+	{
+		if (bPreviousShouldHideUI && !PreviousSelectedRenderGridJobIds.IsEmpty())
+		{
+			PreviousSelectedRenderGridJobIds.Empty();
+			OnRenderGridJobsSelectionChanged().Broadcast();
+		}
+		else
+		{
+			if ((JobSelection->SelectedRenderGridJobIds.Num() != PreviousSelectedRenderGridJobIds.Num()) || (JobSelection->SelectedRenderGridJobIds.Intersect(PreviousSelectedRenderGridJobIds).Num() != PreviousSelectedRenderGridJobIds.Num()))
+			{
+				PreviousSelectedRenderGridJobIds = JobSelection->SelectedRenderGridJobIds;
+				OnRenderGridJobsSelectionChanged().Broadcast();
+			}
+		}
+	}
+
 	if ((DebuggingTimeInSecondsRemaining > 0) && !bIsDebugging)
 	{
 		DebuggingTimeInSecondsRemaining -= DeltaTime;
@@ -546,36 +575,28 @@ void UE::RenderGrid::Private::FRenderGridEditor::BatchRenderListAction()
 
 	if (URenderGrid* Instance = GetInstance(); IsValid(Instance))
 	{
-		if (Instance->GetRenderGridJobs().Num() <= 0)
-		{
-			const FText TitleText = LOCTEXT("NoJobsToRenderTitle", "No Jobs To Render");
-			FMessageDialog::Open(
-				EAppMsgType::Ok,
-				LOCTEXT("NoJobsToRenderText", "There are no jobs in this grid, and so nothing can be rendered. Please add a job and try again."),
-				&TitleText);
-			return;
-		}
-
-		if (URenderGridQueue* RenderQueue = IRenderGridModule::Get().GetManager().CreateBatchRenderQueue(Instance); IsValid(RenderQueue))
+		if (URenderGridQueue* RenderQueue = Instance->Render(); IsValid(RenderQueue))
 		{
-			RenderQueue->OnExecuteFinished().AddRaw(this, &FRenderGridEditor::OnBatchRenderListActionFinished);
 			BatchRenderQueue = TStrongObjectPtr(RenderQueue);
 
-			SetIsDebugging(true);
-			RenderQueue->Execute();
+			TWeakPtr<FRenderGridEditor> ThisWeakPtr = SharedThis(this);
+			RenderQueue->OnExecuteFinished().AddLambda([ThisWeakPtr](URenderGridQueue* Queue, bool bSuccess)
+			{
+				if (const TSharedPtr<FRenderGridEditor> This = ThisWeakPtr.Pin())
+				{
+					This->OnBatchRenderListActionFinished();
+				}
+			});
 
-			OnRenderGridBatchRenderingStarted().Broadcast(RenderQueue);
+			SetIsDebugging(true);
 		}
 	}
 }
 
-void UE::RenderGrid::Private::FRenderGridEditor::OnBatchRenderListActionFinished(URenderGridQueue* Queue, bool bSuccess)
+void UE::RenderGrid::Private::FRenderGridEditor::OnBatchRenderListActionFinished()
 {
-	SetIsDebugging(false);
-
-	URenderGridQueue* FinishedRenderQueue = BatchRenderQueue.Get();
 	BatchRenderQueue = nullptr;
-	OnRenderGridBatchRenderingFinished().Broadcast(FinishedRenderQueue);
+	SetIsDebugging(false);
 }
 
 void UE::RenderGrid::Private::FRenderGridEditor::UndoAction()
@@ -600,7 +621,7 @@ void UE::RenderGrid::Private::FRenderGridEditor::ExtendMenu()
 
 	AddMenuExtender(MenuExtender);
 
-	// add extensible menu if exists
+	// Add extensible menu if exists
 	IRenderGridEditorModule& RenderGridEditorModule = IRenderGridEditorModule::Get();
 	AddMenuExtender(RenderGridEditorModule.GetMenuExtensibilityManager()->GetAllExtenders(GetToolkitCommands(), GetEditingObjects()));
 }
@@ -644,9 +665,6 @@ void UE::RenderGrid::Private::FRenderGridEditor::DestroyInstance()
 	{
 		// Execute the blueprint event
 		Instance->EndEditor();
-
-		// Garbage collection
-		Instance->MarkAsGarbage();
 	}
 
 	if (URenderGridBlueprint* PreviewBlueprint = PreviewBlueprintWeakPtr.Get(); IsValid(PreviewBlueprint))
diff --git a/Engine/Plugins/Experimental/RenderGrid/Source/RenderGridEditor/Private/Toolkit/RenderGridEditor.h b/Engine/Plugins/Experimental/RenderGrid/Source/RenderGridEditor/Private/Toolkit/RenderGridEditor.h
--- a/Engine/Plugins/Experimental/RenderGrid/Source/RenderGridEditor/Private/Toolkit/RenderGridEditor.h
+++ b/Engine/Plugins/Experimental/RenderGrid/Source/RenderGridEditor/Private/Toolkit/RenderGridEditor.h
@@ -57,6 +57,11 @@ namespace UE::RenderGrid::Private
 		/** The time it should remain in debugging mode after it has been turned off. */
 		static constexpr float TimeInSecondsToRemainDebugging = 4.0f;
 
+	protected:
+		/** The currently selected render grid jobs. It's static so multiple open render grid editors won't cause bugs with each other. */
+		static TStrongObjectPtr<URenderGridJobSelection> RenderGridJobSelection;
+
+
 	public:
 		void InitRenderGridEditor(const EToolkitMode::Type Mode, const TSharedPtr<IToolkitHost>& InitToolkitHost, URenderGridBlueprint* InRenderGridBlueprint);
 
@@ -142,7 +147,7 @@ namespace UE::RenderGrid::Private
 		void BatchRenderListAction();
 
 		/** The callback for when the batch render list action finishes. */
-		void OnBatchRenderListActionFinished(URenderGridQueue* Queue, bool bSuccess);
+		void OnBatchRenderListActionFinished();
 
 		/** Undo the last action. */
 		void UndoAction();
@@ -186,9 +191,6 @@ namespace UE::RenderGrid::Private
 		/** The current render grid instance that's visible in the editor. */
 		mutable TWeakObjectPtr<URenderGrid> RenderGridWeakPtr;
 
-		/** The currently selected render grid jobs. */
-		TStrongObjectPtr<URenderGridJobSelection> RenderGridJobSelection;
-
 		/** True if it should call BatchRenderListAction() next frame. */
 		bool bRunRenderNewBatch;
 
@@ -203,5 +205,11 @@ namespace UE::RenderGrid::Private
 
 		/** True if the graph editor is currently in debugging mode. */
 		bool bIsDebugging;
+
+		/** The value of ShouldHideUI() during the last tick. */
+		bool bPreviousShouldHideUI;
+
+		/** The value of RenderGridJobSelection.Get()->SelectedRenderGridJobIds() during the last tick. */
+		TSet<FGuid> PreviousSelectedRenderGridJobIds;
 	};
 }
diff --git a/Engine/Plugins/Experimental/RenderGrid/Source/RenderGridEditor/Private/UI/Components/SRenderGridDragHandle.h b/Engine/Plugins/Experimental/RenderGrid/Source/RenderGridEditor/Private/UI/Components/SRenderGridDragHandle.h
--- a/Engine/Plugins/Experimental/RenderGrid/Source/RenderGridEditor/Private/UI/Components/SRenderGridDragHandle.h
+++ b/Engine/Plugins/Experimental/RenderGrid/Source/RenderGridEditor/Private/UI/Components/SRenderGridDragHandle.h
@@ -53,7 +53,7 @@ public:
 	{
 		if (MouseEvent.IsMouseButtonDown(EKeys::LeftMouseButton))
 		{
-			if (TSharedPtr<FDragDropOperation> DragDropOp = CreateDragDropOperation())
+			if (const TSharedPtr<FDragDropOperation> DragDropOp = CreateDragDropOperation())
 			{
 				return FReply::Handled().BeginDragDrop(DragDropOp.ToSharedRef());
 			}
diff --git a/Engine/Plugins/Experimental/RenderGrid/Source/RenderGridEditor/Private/UI/Components/SRenderGridRemoteControlField.cpp b/Engine/Plugins/Experimental/RenderGrid/Source/RenderGridEditor/Private/UI/Components/SRenderGridRemoteControlField.cpp
--- a/Engine/Plugins/Experimental/RenderGrid/Source/RenderGridEditor/Private/UI/Components/SRenderGridRemoteControlField.cpp
+++ b/Engine/Plugins/Experimental/RenderGrid/Source/RenderGridEditor/Private/UI/Components/SRenderGridRemoteControlField.cpp
@@ -157,7 +157,7 @@ TSharedRef<SWidget> UE::RenderGrid::Private::SRenderGridRemoteControlField::Cons
 			Generator = FModuleManager::GetModuleChecked<FPropertyEditorModule>("PropertyEditor").CreatePropertyRowGenerator(Args);
 			Generator->SetObjects({Objects[0]});
 
-			if (TSharedPtr<IDetailTreeNode> Node = RenderGridWidgetUtils::FindNode(Generator->GetRootTreeNodes(), Field->FieldPathInfo.ToPathPropertyString(), RenderGridWidgetUtils::ERenderGridFindNodeMethod::Path))
+			if (const TSharedPtr<IDetailTreeNode> Node = RenderGridWidgetUtils::FindNode(Generator->GetRootTreeNodes(), Field->FieldPathInfo.ToPathPropertyString(), RenderGridWidgetUtils::ERenderGridFindNodeMethod::Path))
 			{
 				TArray<TSharedRef<IDetailTreeNode>> ChildNodes;
 				Node->GetChildren(ChildNodes);
diff --git a/Engine/Plugins/Experimental/RenderGrid/Source/RenderGridEditor/Private/UI/SRenderGrid.cpp b/Engine/Plugins/Experimental/RenderGrid/Source/RenderGridEditor/Private/UI/SRenderGrid.cpp
--- a/Engine/Plugins/Experimental/RenderGrid/Source/RenderGridEditor/Private/UI/SRenderGrid.cpp
+++ b/Engine/Plugins/Experimental/RenderGrid/Source/RenderGridEditor/Private/UI/SRenderGrid.cpp
@@ -19,7 +19,7 @@ void UE::RenderGrid::Private::SRenderGrid::Tick(const FGeometry&, const double,
 		{
 			if (URenderGrid* Grid = BlueprintEditor->GetInstance())
 			{
-				if (!IsValid(Grid) || BlueprintEditor->IsBatchRendering())
+				if (!IsValid(Grid) || BlueprintEditor->ShouldHideUI())
 				{
 					Grid = nullptr;
 				}
@@ -48,6 +48,7 @@ void UE::RenderGrid::Private::SRenderGrid::Construct(const FArguments& InArgs, T
 	DetailsViewArgs.NotifyHook = this;
 
 	RenderGridDetailsView = PropertyEditorModule.CreateDetailView(DetailsViewArgs);
+	RenderGridSettingsDetailsView = PropertyEditorModule.CreateDetailView(DetailsViewArgs);
 	RenderGridDefaultsDetailsView = PropertyEditorModule.CreateDetailView(DetailsViewArgs);
 
 	SetRenderGrid(InBlueprintEditor->GetInstance());
@@ -62,6 +63,12 @@ void UE::RenderGrid::Private::SRenderGrid::Construct(const FArguments& InArgs, T
 			RenderGridDetailsView->AsShared()
 		]
 
+		+ SVerticalBox::Slot()
+		.AutoHeight()
+		[
+			RenderGridSettingsDetailsView->AsShared()
+		]
+
 		+ SVerticalBox::Slot()
 		.AutoHeight()
 		[
@@ -89,6 +96,7 @@ void UE::RenderGrid::Private::SRenderGrid::SetRenderGrid(URenderGrid* RenderGrid
 {
 	RenderGridWeakPtr = RenderGrid;
 	RenderGridDetailsView->SetObject(RenderGrid);
+	RenderGridSettingsDetailsView->SetObject(IsValid(RenderGrid) ? RenderGrid->GetSettingsObject() : nullptr);
 	RenderGridDefaultsDetailsView->SetObject(IsValid(RenderGrid) ? RenderGrid->GetDefaultsObject() : nullptr);
 }
 
diff --git a/Engine/Plugins/Experimental/RenderGrid/Source/RenderGridEditor/Private/UI/SRenderGrid.h b/Engine/Plugins/Experimental/RenderGrid/Source/RenderGridEditor/Private/UI/SRenderGrid.h
--- a/Engine/Plugins/Experimental/RenderGrid/Source/RenderGridEditor/Private/UI/SRenderGrid.h
+++ b/Engine/Plugins/Experimental/RenderGrid/Source/RenderGridEditor/Private/UI/SRenderGrid.h
@@ -50,6 +50,9 @@ namespace UE::RenderGrid::Private
 		/** A reference to the details view. */
 		TSharedPtr<IDetailsView> RenderGridDetailsView;
 
+		/** A reference to the details view. */
+		TSharedPtr<IDetailsView> RenderGridSettingsDetailsView;
+
 		/** A reference to the details view. */
 		TSharedPtr<IDetailsView> RenderGridDefaultsDetailsView;
 	};
diff --git a/Engine/Plugins/Experimental/RenderGrid/Source/RenderGridEditor/Private/UI/SRenderGridJob.cpp b/Engine/Plugins/Experimental/RenderGrid/Source/RenderGridEditor/Private/UI/SRenderGridJob.cpp
--- a/Engine/Plugins/Experimental/RenderGrid/Source/RenderGridEditor/Private/UI/SRenderGridJob.cpp
+++ b/Engine/Plugins/Experimental/RenderGrid/Source/RenderGridEditor/Private/UI/SRenderGridJob.cpp
@@ -25,9 +25,8 @@ void UE::RenderGrid::Private::SRenderGridJob::Construct(const FArguments& InArgs
 	DetailsView->SetObject(nullptr);
 
 	Refresh();
+	InBlueprintEditor->OnRenderGridShouldHideUIChanged().AddSP(this, &SRenderGridJob::Refresh);
 	InBlueprintEditor->OnRenderGridJobsSelectionChanged().AddSP(this, &SRenderGridJob::Refresh);
-	InBlueprintEditor->OnRenderGridBatchRenderingStarted().AddSP(this, &SRenderGridJob::OnBatchRenderingStarted);
-	InBlueprintEditor->OnRenderGridBatchRenderingFinished().AddSP(this, &SRenderGridJob::OnBatchRenderingFinished);
 
 	ChildSlot
 	[
@@ -55,7 +54,7 @@ void UE::RenderGrid::Private::SRenderGridJob::Refresh()
 		if (const TSharedPtr<IRenderGridEditor> BlueprintEditor = BlueprintEditorWeakPtr.Pin())
 		{
 			TArray<TWeakObjectPtr<UObject>> WeakSelectedJobs;
-			if (!BlueprintEditor->IsBatchRendering())
+			if (!BlueprintEditor->ShouldHideUI())
 			{
 				if (const TArray<URenderGridJob*> SelectedJobs = BlueprintEditor->GetSelectedRenderGridJobs(); (SelectedJobs.Num() == 1))
 				{
diff --git a/Engine/Plugins/Experimental/RenderGrid/Source/RenderGridEditor/Private/UI/SRenderGridJob.h b/Engine/Plugins/Experimental/RenderGrid/Source/RenderGridEditor/Private/UI/SRenderGridJob.h
--- a/Engine/Plugins/Experimental/RenderGrid/Source/RenderGridEditor/Private/UI/SRenderGridJob.h
+++ b/Engine/Plugins/Experimental/RenderGrid/Source/RenderGridEditor/Private/UI/SRenderGridJob.h
@@ -40,10 +40,6 @@ namespace UE::RenderGrid::Private
 		/** Updates the details view. */
 		void Refresh();
 
-	private:
-		void OnBatchRenderingStarted(URenderGridQueue* Queue) { Refresh(); }
-		void OnBatchRenderingFinished(URenderGridQueue* Queue) { Refresh(); }
-
 	private:
 		/** A reference to the blueprint editor that owns the render grid instance. */
 		TWeakPtr<IRenderGridEditor> BlueprintEditorWeakPtr;
diff --git a/Engine/Plugins/Experimental/RenderGrid/Source/RenderGridEditor/Private/UI/SRenderGridJobList.cpp b/Engine/Plugins/Experimental/RenderGrid/Source/RenderGridEditor/Private/UI/SRenderGridJobList.cpp
--- a/Engine/Plugins/Experimental/RenderGrid/Source/RenderGridEditor/Private/UI/SRenderGridJobList.cpp
+++ b/Engine/Plugins/Experimental/RenderGrid/Source/RenderGridEditor/Private/UI/SRenderGridJobList.cpp
@@ -54,7 +54,7 @@ void UE::RenderGrid::Private::SRenderGridJobList::Tick(const FGeometry&, const d
 {
 	if (const TSharedPtr<IRenderGridEditor> BlueprintEditor = BlueprintEditorWeakPtr.Pin())
 	{
-		if (RenderGridWeakPtr != BlueprintEditor->GetInstance())
+		if ((RenderGridWeakPtr != BlueprintEditor->GetInstance()) || (BlueprintEditor->IsBatchRendering() != bPreviousIsBatchRendering))
 		{
 			Refresh();
 		}
@@ -73,8 +73,7 @@ void UE::RenderGrid::Private::SRenderGridJobList::Construct(const FArguments& In
 	Refresh();
 	InBlueprintEditor->OnRenderGridJobCreated().AddSP(this, &SRenderGridJobList::OnRenderGridJobCreated);
 	InBlueprintEditor->OnRenderGridChanged().AddSP(this, &SRenderGridJobList::Refresh);
-	InBlueprintEditor->OnRenderGridBatchRenderingStarted().AddSP(this, &SRenderGridJobList::OnBatchRenderingStarted);
-	InBlueprintEditor->OnRenderGridBatchRenderingFinished().AddSP(this, &SRenderGridJobList::OnBatchRenderingFinished);
+	InBlueprintEditor->OnRenderGridJobsSelectionChanged().AddSP(this, &SRenderGridJobList::Refresh);
 	FCoreUObjectDelegates::OnObjectModified.AddSP(this, &SRenderGridJobList::OnObjectModified);
 
 	ChildSlot
@@ -313,7 +312,7 @@ void UE::RenderGrid::Private::SRenderGridJobList::Refresh()
 {
 	if (const TSharedPtr<IRenderGridEditor> BlueprintEditor = BlueprintEditorWeakPtr.Pin())
 	{
-		const bool bIsBatchRendering = BlueprintEditor->IsBatchRendering();// show all jobs during a batch render, ignore the search bar
+		bPreviousIsBatchRendering = BlueprintEditor->IsBatchRendering(); // show all jobs during a batch render, ignore the search bar
 		const FString SearchBarContent = (!RenderGridSearchBox.IsValid() ? TEXT("") : RenderGridSearchBox->GetText().ToString());
 
 		RenderGridJobs.Empty();
@@ -322,7 +321,7 @@ void UE::RenderGrid::Private::SRenderGridJobList::Refresh()
 		{
 			for (URenderGridJob* Job : Grid->GetRenderGridJobs())
 			{
-				if (bIsBatchRendering || Job->MatchesSearchTerm(SearchBarContent))
+				if (bPreviousIsBatchRendering || Job->MatchesSearchTerm(SearchBarContent))
 				{
 					RenderGridJobs.Add(Job);
 				}
@@ -332,7 +331,7 @@ void UE::RenderGrid::Private::SRenderGridJobList::Refresh()
 		RefreshHeaderEnabledCheckbox();
 
 		RemoveRenderStatusColumn();
-		if (bIsBatchRendering)
+		if (bPreviousIsBatchRendering)
 		{
 			AddRenderStatusColumn();
 		}
@@ -347,7 +346,6 @@ void UE::RenderGrid::Private::SRenderGridJobList::Refresh()
 			});
 			RenderGridJobListWidget->ClearSelection();
 			RenderGridJobListWidget->SetItemSelection(SelectedJobs, true);
-			BlueprintEditor->SetSelectedRenderGridJobs(SelectedJobs);
 		}
 	}
 }
diff --git a/Engine/Plugins/Experimental/RenderGrid/Source/RenderGridEditor/Private/UI/SRenderGridJobList.h b/Engine/Plugins/Experimental/RenderGrid/Source/RenderGridEditor/Private/UI/SRenderGridJobList.h
--- a/Engine/Plugins/Experimental/RenderGrid/Source/RenderGridEditor/Private/UI/SRenderGridJobList.h
+++ b/Engine/Plugins/Experimental/RenderGrid/Source/RenderGridEditor/Private/UI/SRenderGridJobList.h
@@ -74,8 +74,6 @@ namespace UE::RenderGrid::Private
 
 	private:
 		void OnObjectModified(UObject* Object);
-		void OnBatchRenderingStarted(URenderGridQueue* Queue) { Refresh(); }
-		void OnBatchRenderingFinished(URenderGridQueue* Queue) { Refresh(); }
 		void OnSearchBarTextChanged(const FText& Text) { Refresh(); }
 
 	private:
@@ -96,6 +94,9 @@ namespace UE::RenderGrid::Private
 
 		/** The header checkbox for the enable/disable column. */
 		TSharedPtr<SCheckBox> RenderGridJobEnabledHeaderCheckbox;
+
+		/** The value of BlueprintEditorWeakPtr.Get()->IsBatchRendering() during the last Refresh(). */
+		bool bPreviousIsBatchRendering;
 	};
 
 
diff --git a/Engine/Plugins/Experimental/RenderGrid/Source/RenderGridEditor/Private/UI/SRenderGridProps.cpp b/Engine/Plugins/Experimental/RenderGrid/Source/RenderGridEditor/Private/UI/SRenderGridProps.cpp
--- a/Engine/Plugins/Experimental/RenderGrid/Source/RenderGridEditor/Private/UI/SRenderGridProps.cpp
+++ b/Engine/Plugins/Experimental/RenderGrid/Source/RenderGridEditor/Private/UI/SRenderGridProps.cpp
@@ -38,9 +38,8 @@ void UE::RenderGrid::Private::SRenderGridProps::Construct(const FArguments& InAr
 	Refresh();
 	//InBlueprintEditor->OnRenderGridChanged().AddSP(this, &SRenderGridProps::Refresh); // should technically be correct, but causes the remote control properties to bug out
 	InBlueprintEditor->OnRenderGridPreviewRenderedFirstTimeSinceAppStart().AddSP(this, &SRenderGridProps::Refresh);
+	InBlueprintEditor->OnRenderGridShouldHideUIChanged().AddSP(this, &SRenderGridProps::Refresh);
 	InBlueprintEditor->OnRenderGridJobsSelectionChanged().AddSP(this, &SRenderGridProps::Refresh);
-	InBlueprintEditor->OnRenderGridBatchRenderingStarted().AddSP(this, &SRenderGridProps::OnBatchRenderingStarted);
-	InBlueprintEditor->OnRenderGridBatchRenderingFinished().AddSP(this, &SRenderGridProps::OnBatchRenderingFinished);
 
 	ChildSlot
 	[
@@ -60,7 +59,7 @@ void UE::RenderGrid::Private::SRenderGridProps::Refresh()
 
 	if (const TSharedPtr<IRenderGridEditor> BlueprintEditor = BlueprintEditorWeakPtr.Pin())
 	{
-		if (!BlueprintEditor->IsBatchRendering())
+		if (!BlueprintEditor->ShouldHideUI())
 		{
 			if (URenderGrid* Grid = BlueprintEditor->GetInstance(); IsValid(Grid))
 			{
diff --git a/Engine/Plugins/Experimental/RenderGrid/Source/RenderGridEditor/Private/UI/SRenderGridProps.h b/Engine/Plugins/Experimental/RenderGrid/Source/RenderGridEditor/Private/UI/SRenderGridProps.h
--- a/Engine/Plugins/Experimental/RenderGrid/Source/RenderGridEditor/Private/UI/SRenderGridProps.h
+++ b/Engine/Plugins/Experimental/RenderGrid/Source/RenderGridEditor/Private/UI/SRenderGridProps.h
@@ -35,10 +35,6 @@ namespace UE::RenderGrid::Private
 		/** Refreshes the content of this widget. */
 		void Refresh();
 
-	private:
-		void OnBatchRenderingStarted(URenderGridQueue* Queue) { Refresh(); }
-		void OnBatchRenderingFinished(URenderGridQueue* Queue) { Refresh(); }
-
 	private:
 		/** A reference to the blueprint editor that owns the render grid instance. */
 		TWeakPtr<IRenderGridEditor> BlueprintEditorWeakPtr;
diff --git a/Engine/Plugins/Experimental/RenderGrid/Source/RenderGridEditor/Private/UI/SRenderGridViewer.cpp b/Engine/Plugins/Experimental/RenderGrid/Source/RenderGridEditor/Private/UI/SRenderGridViewer.cpp
--- a/Engine/Plugins/Experimental/RenderGrid/Source/RenderGridEditor/Private/UI/SRenderGridViewer.cpp
+++ b/Engine/Plugins/Experimental/RenderGrid/Source/RenderGridEditor/Private/UI/SRenderGridViewer.cpp
@@ -24,8 +24,7 @@ void UE::RenderGrid::Private::SRenderGridViewer::Construct(const FArguments& InA
 		.BorderImage(new FSlateNoResource());
 
 	Refresh();
-	InBlueprintEditor->OnRenderGridBatchRenderingStarted().AddSP(this, &SRenderGridViewer::OnBatchRenderingStarted);
-	InBlueprintEditor->OnRenderGridBatchRenderingFinished().AddSP(this, &SRenderGridViewer::OnBatchRenderingFinished);
+	InBlueprintEditor->OnRenderGridShouldHideUIChanged().AddSP(this, &SRenderGridViewer::Refresh);
 
 	ChildSlot
 	[
@@ -99,7 +98,7 @@ void UE::RenderGrid::Private::SRenderGridViewer::Refresh()
 	}
 	if (const TSharedPtr<IRenderGridEditor> BlueprintEditor = BlueprintEditorWeakPtr.Pin())
 	{
-		ERenderGridViewerMode CurrentViewerMode = (BlueprintEditor->IsBatchRendering() ? ERenderGridViewerMode::None : ViewerMode);
+		ERenderGridViewerMode CurrentViewerMode = (BlueprintEditor->ShouldHideUI() ? ERenderGridViewerMode::None : ViewerMode);
 		if (CurrentViewerMode == CachedViewerMode)
 		{
 			return;
diff --git a/Engine/Plugins/Experimental/RenderGrid/Source/RenderGridEditor/Private/UI/SRenderGridViewer.h b/Engine/Plugins/Experimental/RenderGrid/Source/RenderGridEditor/Private/UI/SRenderGridViewer.h
--- a/Engine/Plugins/Experimental/RenderGrid/Source/RenderGridEditor/Private/UI/SRenderGridViewer.h
+++ b/Engine/Plugins/Experimental/RenderGrid/Source/RenderGridEditor/Private/UI/SRenderGridViewer.h
@@ -48,10 +48,6 @@ namespace UE::RenderGrid::Private
 		/** Refreshes the content of this widget. */
 		void Refresh();
 
-	private:
-		void OnBatchRenderingStarted(URenderGridQueue* Queue) { Refresh(); }
-		void OnBatchRenderingFinished(URenderGridQueue* Queue) { Refresh(); }
-
 	private:
 		/** A reference to the blueprint editor that owns the render grid instance. */
 		TWeakPtr<IRenderGridEditor> BlueprintEditorWeakPtr;
diff --git a/Engine/Plugins/Experimental/RenderGrid/Source/RenderGridEditor/Public/IRenderGridEditor.h b/Engine/Plugins/Experimental/RenderGrid/Source/RenderGridEditor/Public/IRenderGridEditor.h
--- a/Engine/Plugins/Experimental/RenderGrid/Source/RenderGridEditor/Public/IRenderGridEditor.h
+++ b/Engine/Plugins/Experimental/RenderGrid/Source/RenderGridEditor/Public/IRenderGridEditor.h
@@ -4,6 +4,7 @@
 
 #include "CoreMinimal.h"
 #include "BlueprintEditor.h"
+#include "RenderGrid/RenderGridQueue.h"
 
 
 class URenderGridQueue;
@@ -38,11 +39,14 @@ namespace UE::RenderGrid
 		virtual TSharedPtr<Private::FRenderGridBlueprintEditorToolbar> GetRenderGridToolbarBuilder() = 0;
 
 		/** Returns whether it is currently rendering or playing (so changes in the level and such should be ignored). */
-		virtual bool IsCurrentlyRenderingOrPlaying() const { return IsBatchRendering() || IsPreviewRendering() || IsValid(GEditor->PlayWorld); }
+		virtual bool IsCurrentlyRenderingOrPlaying() const { return IsBatchRendering() || IsPreviewRendering() || URenderGridQueue::IsExecutingAny() || IsValid(GEditor->PlayWorld); }
 
 		/** Returns whether it can currently render (like a preview render or a batch render) or not. */
 		virtual bool CanCurrentlyRender() const { return !IsCurrentlyRenderingOrPlaying(); }
 
+		/** Returns whether it is should currently hide the UI for things to not cause any bugs (like during rendering etc). */
+		virtual bool ShouldHideUI() const { return (IsCurrentlyRenderingOrPlaying() && !IsPreviewRendering()); }
+
 		/** Returns whether it is currently batch rendering or not. */
 		virtual bool IsBatchRendering() const = 0;
 
@@ -76,11 +80,8 @@ namespace UE::RenderGrid
 		DECLARE_MULTICAST_DELEGATE(FOnRenderGridJobsSelectionChanged);
 		virtual FOnRenderGridJobsSelectionChanged& OnRenderGridJobsSelectionChanged() { return OnRenderGridJobsSelectionChangedDelegate; }
 
-		DECLARE_MULTICAST_DELEGATE_OneParam(FOnRenderGridBatchRenderingStarted, URenderGridQueue*);
-		virtual FOnRenderGridBatchRenderingStarted& OnRenderGridBatchRenderingStarted() { return OnRenderGridBatchRenderingStartedDelegate; }
-
-		DECLARE_MULTICAST_DELEGATE_OneParam(FOnRenderGridBatchRenderingFinished, URenderGridQueue*);
-		virtual FOnRenderGridBatchRenderingFinished& OnRenderGridBatchRenderingFinished() { return OnRenderGridBatchRenderingFinishedDelegate; }
+		DECLARE_MULTICAST_DELEGATE(FOnRenderGridShouldHideUIChanged);
+		virtual FOnRenderGridShouldHideUIChanged& OnRenderGridShouldHideUIChanged() { return OnRenderGridShouldHideUIChangedDelegate; }
 
 		DECLARE_MULTICAST_DELEGATE(FOnRenderGridPreviewRenderedFirstTimeSinceAppStart)
 		virtual FOnRenderGridPreviewRenderedFirstTimeSinceAppStart& OnRenderGridPreviewRenderedFirstTimeSinceAppStart() { return OnRenderGridPreviewRenderedFirstTimeSinceAppStartDelegate; }
@@ -95,11 +96,8 @@ namespace UE::RenderGrid
 		/** The delegate for when the selection of render grid jobs changed. */
 		FOnRenderGridJobsSelectionChanged OnRenderGridJobsSelectionChangedDelegate;
 
-		/** The delegate for when batch rendering of a render grid has started. */
-		FOnRenderGridBatchRenderingStarted OnRenderGridBatchRenderingStartedDelegate;
-
-		/** The delegate for when batch rendering of a render grid has finished, successfully or not. */
-		FOnRenderGridBatchRenderingFinished OnRenderGridBatchRenderingFinishedDelegate;
+		/** The delegate for when the the result of ShouldHideUI() has changed. */
+		FOnRenderGridShouldHideUIChanged OnRenderGridShouldHideUIChangedDelegate;
 
 		/** The delegate for the first time (since the start of this application) that the preview is rendered. */
 		FOnRenderGridPreviewRenderedFirstTimeSinceAppStart OnRenderGridPreviewRenderedFirstTimeSinceAppStartDelegate;
