RenderGrid:
- Replaced the bottom-right blueprint-editor text from "RENDER PAGES" to "RENDER GRID".

diff --git a/Engine/Plugins/Experimental/RenderGrid/Source/RenderGridEditor/Private/Toolkit/RenderGridEditor.cpp b/Engine/Plugins/Experimental/RenderGrid/Source/RenderGridEditor/Private/Toolkit/RenderGridEditor.cpp
--- a/Engine/Plugins/Experimental/RenderGrid/Source/RenderGridEditor/Private/Toolkit/RenderGridEditor.cpp
+++ b/Engine/Plugins/Experimental/RenderGrid/Source/RenderGridEditor/Private/Toolkit/RenderGridEditor.cpp
@@ -172,7 +172,7 @@ FGraphAppearanceInfo UE::RenderGrid::Private::FRenderGridEditor::GetGraphAppeara
 	FGraphAppearanceInfo AppearanceInfo = FBlueprintEditor::GetGraphAppearance(InGraph);
 	if (GetBlueprintObj()->IsA(URenderGridBlueprint::StaticClass()))
 	{
-		AppearanceInfo.CornerText = LOCTEXT("AppearanceCornerText_RenderGrid", "RENDER PAGES");
+		AppearanceInfo.CornerText = LOCTEXT("AppearanceCornerText_RenderGrid", "RENDER GRID");
 	}
 	return AppearanceInfo;
 }
