Fixed a crash in headless mode caused by the render grid plugin trying to access the slate renderer while the slate renderer isn't available.

diff --git a/Engine/Plugins/Experimental/RenderGrid/Source/RenderGridEditor/Private/Styles/RenderGridEditorStyle.cpp b/Engine/Plugins/Experimental/RenderGrid/Source/RenderGridEditor/Private/Styles/RenderGridEditorStyle.cpp
--- a/Engine/Plugins/Experimental/RenderGrid/Source/RenderGridEditor/Private/Styles/RenderGridEditorStyle.cpp
+++ b/Engine/Plugins/Experimental/RenderGrid/Source/RenderGridEditor/Private/Styles/RenderGridEditorStyle.cpp
@@ -34,6 +34,17 @@ void UE::RenderGrid::Private::FRenderGridEditorStyle::Shutdown()
 	StyleInstance.Reset();
 }
 
+void UE::RenderGrid::Private::FRenderGridEditorStyle::ReloadTextures()
+{
+	if (FSlateApplication::IsInitialized())
+	{
+		if (FSlateRenderer* SlateRenderer = FSlateApplication::Get().GetRenderer())
+		{
+			SlateRenderer->ReloadTextureResources();
+		}
+	}
+}
+
 FName UE::RenderGrid::Private::FRenderGridEditorStyle::GetStyleSetName()
 {
 	static FName StyleSetName(TEXT("RenderGridEditorStyle"));
@@ -102,8 +113,3 @@ TSharedRef<FSlateStyleSet> UE::RenderGrid::Private::FRenderGridEditorStyle::Crea
 #undef BOX_BRUSH
 #undef BORDER_BRUSH
 #undef IMAGE_BRUSH_SVG
-
-void UE::RenderGrid::Private::FRenderGridEditorStyle::ReloadTextures()
-{
-	FSlateApplication::Get().GetRenderer()->ReloadTextureResources();
-}
diff --git a/Engine/Plugins/Experimental/RenderGrid/Source/RenderGridEditor/Private/UI/Components/SRenderGridFileSelectorTextBlock.cpp b/Engine/Plugins/Experimental/RenderGrid/Source/RenderGridEditor/Private/UI/Components/SRenderGridFileSelectorTextBlock.cpp
--- a/Engine/Plugins/Experimental/RenderGrid/Source/RenderGridEditor/Private/UI/Components/SRenderGridFileSelectorTextBlock.cpp
+++ b/Engine/Plugins/Experimental/RenderGrid/Source/RenderGridEditor/Private/UI/Components/SRenderGridFileSelectorTextBlock.cpp
@@ -75,24 +75,27 @@ FReply UE::RenderGrid::Private::SRenderGridFileSelectorTextBlock::OnOpenDirector
 {
 	if (IDesktopPlatform* DesktopPlatform = FDesktopPlatformModule::Get())
 	{
-		TSharedPtr<SWindow> ParentWindow = FSlateApplication::Get().FindWidgetWindow(AsShared());
-		void* ParentWindowHandle = (ParentWindow.IsValid() && ParentWindow->GetNativeWindow().IsValid()) ? ParentWindow->GetNativeWindow()->GetOSWindowHandle() : nullptr;
-
-		FString OutFolderPath;
-		const bool bFolderSelected = DesktopPlatform->OpenDirectoryDialog(
-			ParentWindowHandle,
-			LOCTEXT("Title_OpenDirectoryDialog", "Choose a directory").ToString(),
-			FolderPath.Get(FPaths::ProjectDir() / TEXT("Saved/MovieRenders/")),
-			OutFolderPath
-		);
-
-		if (bFolderSelected)
+		if (FSlateApplication::IsInitialized())
 		{
-			const FText NewText = FText::FromString(OutFolderPath);
-			SetText(NewText);
-			if (OnTextCommittedDelegate.IsBound())
+			TSharedPtr<SWindow> ParentWindow = FSlateApplication::Get().FindWidgetWindow(AsShared());
+			void* ParentWindowHandle = (ParentWindow.IsValid() && ParentWindow->GetNativeWindow().IsValid()) ? ParentWindow->GetNativeWindow()->GetOSWindowHandle() : nullptr;
+
+			FString OutFolderPath;
+			const bool bFolderSelected = DesktopPlatform->OpenDirectoryDialog(
+				ParentWindowHandle,
+				LOCTEXT("Title_OpenDirectoryDialog", "Choose a directory").ToString(),
+				FolderPath.Get(FPaths::ProjectDir() / TEXT("Saved/MovieRenders/")),
+				OutFolderPath
+			);
+
+			if (bFolderSelected)
 			{
-				SetText(OnTextCommittedDelegate.Execute(NewText, ETextCommit::Type::OnEnter));
+				const FText NewText = FText::FromString(OutFolderPath);
+				SetText(NewText);
+				if (OnTextCommittedDelegate.IsBound())
+				{
+					SetText(OnTextCommittedDelegate.Execute(NewText, ETextCommit::Type::OnEnter));
+				}
 			}
 		}
 	}
diff --git a/Engine/Plugins/Experimental/RenderGrid/Source/RenderGridEditor/Private/UI/SRenderGridJobList.cpp b/Engine/Plugins/Experimental/RenderGrid/Source/RenderGridEditor/Private/UI/SRenderGridJobList.cpp
--- a/Engine/Plugins/Experimental/RenderGrid/Source/RenderGridEditor/Private/UI/SRenderGridJobList.cpp
+++ b/Engine/Plugins/Experimental/RenderGrid/Source/RenderGridEditor/Private/UI/SRenderGridJobList.cpp
@@ -227,7 +227,10 @@ FReply UE::RenderGrid::Private::SRenderGridJobList::OnMouseButtonUp(const FGeome
 			if (MenuContent.IsValid() && (MouseEvent.GetEventPath() != nullptr))
 			{
 				FWidgetPath WidgetPath = *MouseEvent.GetEventPath();
-				FSlateApplication::Get().PushMenu(WidgetPath.Widgets.Last().Widget, WidgetPath, MenuContent.ToSharedRef(), MouseEvent.GetScreenSpacePosition(), FPopupTransitionEffect(FPopupTransitionEffect::ContextMenu));
+				if (FSlateApplication::IsInitialized())
+				{
+					FSlateApplication::Get().PushMenu(WidgetPath.Widgets.Last().Widget, WidgetPath, MenuContent.ToSharedRef(), MouseEvent.GetScreenSpacePosition(), FPopupTransitionEffect(FPopupTransitionEffect::ContextMenu));
+				}
 			}
 			return FReply::Handled();
 		}
