RenderGrid plugin:
- Made it possible to render the render grids (or a single given render grid) using the command line, using python.

diff --git a/Engine/Plugins/Experimental/RenderGrid/Source/RenderGrid/Private/RenderGrid/RenderGridQueue.cpp b/Engine/Plugins/Experimental/RenderGrid/Source/RenderGrid/Private/RenderGrid/RenderGridQueue.cpp
--- a/Engine/Plugins/Experimental/RenderGrid/Source/RenderGrid/Private/RenderGrid/RenderGridQueue.cpp
+++ b/Engine/Plugins/Experimental/RenderGrid/Source/RenderGrid/Private/RenderGrid/RenderGridQueue.cpp
@@ -305,12 +305,27 @@ void URenderGridMoviePipelineRenderJob::ExecuteFinished(UMoviePipelineExecutorBa
 
 TQueue<TObjectPtr<URenderGridQueue>> URenderGridQueue::ExecutingQueues;
 URenderGridQueue::FOnExecutionQueueChanged URenderGridQueue::OnExecutionQueueChangedDelegate;
+bool URenderGridQueue::bExitOnCompletion = false;
 
 bool URenderGridQueue::IsExecutingAny()
 {
 	return (GetCurrentlyExecutingQueue() != nullptr);
 }
 
+void URenderGridQueue::CloseEditorOnCompletion()
+{
+	bExitOnCompletion = true;
+	RequestAppExitIfSetToExitOnCompletion();
+}
+
+void URenderGridQueue::RequestAppExitIfSetToExitOnCompletion()
+{
+	if (bExitOnCompletion && !IsExecutingAny())
+	{
+		FPlatformMisc::RequestExit(false);
+	}
+}
+
 URenderGridQueue* URenderGridQueue::GetCurrentlyExecutingQueue()
 {
 	bool bRemovedAny = false;
@@ -355,8 +370,13 @@ void URenderGridQueue::DoNextExecutingQueue()
 		{
 			Queue->Queue->Start(); // can be called twice here, which is fine (once in GetCurrentlyExecutingQueue() if the previous Queue became invalid, and then once again here)
 		}
+		OnExecutionQueueChanged().Broadcast();
+	}
+	else
+	{
+		OnExecutionQueueChanged().Broadcast();
+		RequestAppExitIfSetToExitOnCompletion();
 	}
-	OnExecutionQueueChanged().Broadcast();
 }
 
 
diff --git a/Engine/Plugins/Experimental/RenderGrid/Source/RenderGrid/Public/RenderGrid/RenderGridQueue.h b/Engine/Plugins/Experimental/RenderGrid/Source/RenderGrid/Public/RenderGrid/RenderGridQueue.h
--- a/Engine/Plugins/Experimental/RenderGrid/Source/RenderGrid/Public/RenderGrid/RenderGridQueue.h
+++ b/Engine/Plugins/Experimental/RenderGrid/Source/RenderGrid/Public/RenderGrid/RenderGridQueue.h
@@ -161,13 +161,21 @@ class RENDERGRID_API URenderGridQueue : public UObject, public FTickableGameObje
 {
 	GENERATED_BODY()
 
-
 public:
+	/** Returns true if it's currently executing any rendering queues. */
 	static bool IsExecutingAny();
 
 	DECLARE_MULTICAST_DELEGATE(FOnExecutionQueueChanged);
+	/** The event that will be fired when the currently executing rendering queue changes (like for example when the previous one has completed its execution). */
 	static FOnExecutionQueueChanged& OnExecutionQueueChanged() { return OnExecutionQueueChangedDelegate; }
 
+	/** Call this function to make it so that the editor will be closed when all rendering queues finish execution. This function has to be only called once. */
+	UFUNCTION(BlueprintCallable, Category="Render Grid Queue", Meta=(Keywords="execution execute finish close stop end done quit"))
+	static void CloseEditorOnCompletion();
+
+private:
+	static void RequestAppExitIfSetToExitOnCompletion();
+
 private:
 	/** Returns the currently executing queue, or NULL if there isn't any currently executing. */
 	static URenderGridQueue* GetCurrentlyExecutingQueue();
@@ -185,6 +193,8 @@ private:
 	/** The delegate for when data in the ExecutingQueues has changed. */
 	static FOnExecutionQueueChanged OnExecutionQueueChangedDelegate;
 
+	/** Whether the editor should be closed when all rendering queues finish their execution. */
+	static bool bExitOnCompletion;
 
 public:
 	//~ Begin FTickableGameObject interface
@@ -243,6 +253,10 @@ public:
 	UFUNCTION(BlueprintPure, Category="Render Grid Queue", Meta=(Keywords="stopped quited exited killed terminated ended succeeded completed finished canceled cancelled"))
 	bool IsCurrentlyRendering() const { return (GetCurrentlyExecutingQueue() == this); }
 
+	/** Retrieves the rendering status of the given render grid job. */
+	UFUNCTION(BlueprintPure, Category="Render Grid Queue", Meta=(Keywords="rendering obtain text"))
+	URenderGrid* GetRenderGrid() const { return RenderGrid; }
+
 	/** Retrieves the rendering status of the given render grid job. */
 	UFUNCTION(BlueprintPure, Category="Render Grid Queue", Meta=(Keywords="rendering progression obtain text"))
 	FString GetJobStatus(URenderGridJob* Job) const;
diff --git a/Engine/Plugins/Experimental/RenderGrid/Source/RenderGridDeveloper/Private/RenderGridDeveloperLibrary.cpp b/Engine/Plugins/Experimental/RenderGrid/Source/RenderGridDeveloper/Private/RenderGridDeveloperLibrary.cpp
--- a/Engine/Plugins/Experimental/RenderGrid/Source/RenderGridDeveloper/Private/RenderGridDeveloperLibrary.cpp
+++ b/Engine/Plugins/Experimental/RenderGrid/Source/RenderGridDeveloper/Private/RenderGridDeveloperLibrary.cpp
@@ -42,3 +42,28 @@ TArray<URenderGrid*> URenderGridDeveloperLibrary::GetAllRenderGridAssets()
 	}
 	return Result;
 }
+
+
+URenderGridBlueprint* URenderGridDeveloperLibrary::GetRenderGridBlueprintAsset(const FString& ObjectPath)
+{
+	if (FSoftObjectPath AssetPath = FSoftObjectPath(ObjectPath); AssetPath.IsValid())
+	{
+		if (URenderGridBlueprint* RenderGridBlueprint = Cast<URenderGridBlueprint>(AssetPath.TryLoad()); IsValid(RenderGridBlueprint))
+		{
+			return RenderGridBlueprint;
+		}
+	}
+	return nullptr;
+}
+
+URenderGrid* URenderGridDeveloperLibrary::GetRenderGridAsset(const FString& ObjectPath)
+{
+	if (URenderGridBlueprint* RenderGridBlueprint = GetRenderGridBlueprintAsset(ObjectPath); IsValid(RenderGridBlueprint))
+	{
+		if (URenderGrid* RenderGrid = RenderGridBlueprint->GetRenderGridWithBlueprintGraph(); IsValid(RenderGrid))
+		{
+			return RenderGrid;
+		}
+	}
+	return nullptr;
+}
diff --git a/Engine/Plugins/Experimental/RenderGrid/Source/RenderGridDeveloper/Public/RenderGridDeveloperLibrary.h b/Engine/Plugins/Experimental/RenderGrid/Source/RenderGridDeveloper/Public/RenderGridDeveloperLibrary.h
--- a/Engine/Plugins/Experimental/RenderGrid/Source/RenderGridDeveloper/Public/RenderGridDeveloperLibrary.h
+++ b/Engine/Plugins/Experimental/RenderGrid/Source/RenderGridDeveloper/Public/RenderGridDeveloperLibrary.h
@@ -33,4 +33,22 @@ public:
 	 */
 	UFUNCTION(BlueprintCallable, Category="Render Grid")
 	static TArray<URenderGrid*> GetAllRenderGridAssets();
+
+
+	/**
+	 * Returns the given render grid asset that exists at the given object path (whether it's on disk or in memory).
+	 * Will load the render grid asset if it's currently unloaded.
+	 * 
+	 * This is a potentially slow operation, so avoid doing this every tick.
+	 */
+	static URenderGridBlueprint* GetRenderGridBlueprintAsset(const FString& ObjectPath);
+
+	/**
+	 * Returns the given render grid asset that exists at the given object path (whether it's on disk or in memory).
+	 * Will load the render grid asset if it's currently unloaded.
+	 * 
+	 * This is a potentially slow operation, so avoid doing this every tick.
+	 */
+	UFUNCTION(BlueprintCallable, Category="Render Grid")
+	static URenderGrid* GetRenderGridAsset(const FString& ObjectPath);
 };
