RenderGrid plugin:
- Made it possible to get and set the remote control values from render grid jobs. Currently you can only retrieve and set a JSON string, which represents the data of remote control. Specific SetAsInt, SetAsColor, etc blueprints and python functions will be added later.

diff --git a/Engine/Plugins/Experimental/RenderGrid/Source/RenderGrid/Private/RenderGrid/RenderGrid.cpp b/Engine/Plugins/Experimental/RenderGrid/Source/RenderGrid/Private/RenderGrid/RenderGrid.cpp
--- a/Engine/Plugins/Experimental/RenderGrid/Source/RenderGrid/Private/RenderGrid/RenderGrid.cpp
+++ b/Engine/Plugins/Experimental/RenderGrid/Source/RenderGrid/Private/RenderGrid/RenderGrid.cpp
@@ -1,15 +1,16 @@
 // Copyright Epic Games, Inc. All Rights Reserved.
 
 #include "RenderGrid/RenderGrid.h"
-#include "RenderGrid/RenderGridManager.h"
-#include "RenderGridUtils.h"
+
 #include "IRenderGridModule.h"
 #include "LevelSequence.h"
 #include "Misc/Base64.h"
-#include "MoviePipelinePrimaryConfig.h"
 #include "MoviePipelineOutputSetting.h"
+#include "MoviePipelinePrimaryConfig.h"
 #include "MovieScene.h"
+#include "RenderGrid/RenderGridManager.h"
 #include "RenderGrid/RenderGridQueue.h"
+#include "RenderGridUtils.h"
 #include "UObject/ObjectSaveContext.h"
 
 
@@ -384,66 +385,250 @@ UMoviePipelineOutputSetting* URenderGridJob::GetRenderPresetOutputSettings() con
 	return nullptr;
 }
 
-bool URenderGridJob::HasRemoteControlValue(const TSharedPtr<FRemoteControlEntity>& RemoteControlEntity) const
+TArray<URemoteControlPreset*> URenderGridJob::GetRemoteControlPresets()
+{
+	TArray<URemoteControlPreset*> Presets;
+	if (URenderGrid* Grid = GetTypedOuter<URenderGrid>(); IsValid(Grid))
+	{
+		if (URenderGridPropsSourceRemoteControl* PropsSource = Grid->GetPropsSource<URenderGridPropsSourceRemoteControl>(); IsValid(PropsSource))
+		{
+			if (URemoteControlPreset* Preset = PropsSource->GetRemoteControlPreset(); IsValid(Preset))
+			{
+				Presets.Add(Preset);
+			}
+		}
+	}
+	/*else
+	{
+		TArray<TSoftObjectPtr<URemoteControlPreset>> RemoteControlPresets;
+		IRemoteControlModule::Get().GetPresets(RemoteControlPresets);
+		for (TSoftObjectPtr<URemoteControlPreset> RemoteControlPresetWeakPtr : RemoteControlPresets)
+		{
+			if (URemoteControlPreset* Preset = RemoteControlPresetWeakPtr.Get(); IsValid(Preset))
+			{
+				Presets.Add(Preset);
+			}
+		}
+	}*/
+	return Presets;
+}
+
+bool URenderGridJob::HasRemoteControlValueBytes(const TSharedPtr<FRemoteControlEntity>& RemoteControlEntity) const
+{
+	if (!RemoteControlEntity.IsValid())
+	{
+		return false;
+	}
+
+	return HasRemoteControlValueBytes(RemoteControlEntity->GetId());
+}
+
+bool URenderGridJob::ConstGetRemoteControlValueBytes(const TSharedPtr<FRemoteControlEntity>& RemoteControlEntity, TArray<uint8>& OutBytes) const
 {
 	if (!RemoteControlEntity.IsValid())
 	{
+		OutBytes.Empty();
 		return false;
 	}
 
-	const FString Key = RemoteControlEntity->GetId().ToString();
-	return !!RemoteControlValues.Find(Key);
+	return ConstGetRemoteControlValueBytes(RemoteControlEntity->GetId(), OutBytes);
 }
 
-bool URenderGridJob::ConstGetRemoteControlValue(const TSharedPtr<FRemoteControlEntity>& RemoteControlEntity, TArray<uint8>& OutBinaryArray) const
+bool URenderGridJob::GetRemoteControlValueBytes(const TSharedPtr<FRemoteControlEntity>& RemoteControlEntity, TArray<uint8>& OutBytes)
 {
+	OutBytes.Empty();
 	if (!RemoteControlEntity.IsValid())
 	{
 		return false;
 	}
 
-	const FString Key = RemoteControlEntity->GetId().ToString();
-	if (const FRenderGridRemoteControlPropertyData* DataPtr = RemoteControlValues.Find(Key))
+	if (FRenderGridRemoteControlPropertyData* DataPtr = RemoteControlValues.Find(RemoteControlEntity->GetId()))
 	{
-		OutBinaryArray.Append((*DataPtr).Bytes);
+		OutBytes.Append((*DataPtr).Bytes);
 		return true;
 	}
-	return false;
+
+	if (!URenderGridPropRemoteControl::GetValueOfEntity(RemoteControlEntity, OutBytes))
+	{
+		return false;
+	}
+	RemoteControlValues.Add(RemoteControlEntity->GetId(), FRenderGridRemoteControlPropertyData(TArray(OutBytes)));
+	return true;
 }
 
-bool URenderGridJob::GetRemoteControlValue(const TSharedPtr<FRemoteControlEntity>& RemoteControlEntity, TArray<uint8>& OutBinaryArray)
+bool URenderGridJob::SetRemoteControlValueBytes(const TSharedPtr<FRemoteControlEntity>& RemoteControlEntity, const TArray<uint8>& Bytes)
 {
 	if (!RemoteControlEntity.IsValid())
 	{
 		return false;
 	}
 
-	const FString Key = RemoteControlEntity->GetId().ToString();
-	if (FRenderGridRemoteControlPropertyData* DataPtr = RemoteControlValues.Find(Key))
+	return SetRemoteControlValueBytes(RemoteControlEntity->GetId(), Bytes);
+}
+
+bool URenderGridJob::HasRemoteControlValueBytes(const FGuid& FieldId) const
+{
+	if (!FieldId.IsValid())
+	{
+		return false;
+	}
+
+	return !!RemoteControlValues.Find(FieldId);
+}
+
+bool URenderGridJob::ConstGetRemoteControlValueBytes(const FGuid& FieldId, TArray<uint8>& OutBytes) const
+{
+	OutBytes.Empty();
+	if (!FieldId.IsValid())
 	{
-		OutBinaryArray.Append((*DataPtr).Bytes);
+		return false;
+	}
+
+	if (const FRenderGridRemoteControlPropertyData* DataPtr = RemoteControlValues.Find(FieldId))
+	{
+		OutBytes.Append((*DataPtr).Bytes);
 		return true;
 	}
+	return false;
+}
 
-	if (!URenderGridPropRemoteControl::GetValueOfEntity(RemoteControlEntity, OutBinaryArray))
+bool URenderGridJob::GetRemoteControlValueBytes(const FGuid& FieldId, TArray<uint8>& OutBytes)
+{
+	OutBytes.Empty();
+	if (!FieldId.IsValid())
 	{
 		return false;
 	}
-	RemoteControlValues.Add(Key, FRenderGridRemoteControlPropertyData(TArray(OutBinaryArray)));
+
+	if (FRenderGridRemoteControlPropertyData* DataPtr = RemoteControlValues.Find(FieldId))
+	{
+		OutBytes.Append((*DataPtr).Bytes);
+		return true;
+	}
+
+	TSharedPtr<FRemoteControlEntity> RemoteControlEntity;
+	for (URemoteControlPreset* RemoteControlPreset : GetRemoteControlPresets())
+	{
+		for (const TWeakPtr<FRemoteControlEntity>& PropWeakPtr : RemoteControlPreset->GetExposedEntities<FRemoteControlEntity>())
+		{
+			if (const TSharedPtr<FRemoteControlEntity> Prop = PropWeakPtr.Pin())
+			{
+				if (Prop->GetId() == FieldId)
+				{
+					RemoteControlEntity = Prop;
+					break;
+				}
+			}
+		}
+	}
+
+	if (!RemoteControlEntity.IsValid() || !URenderGridPropRemoteControl::GetValueOfEntity(RemoteControlEntity, OutBytes))
+	{
+		return false;
+	}
+	RemoteControlValues.Add(FieldId, FRenderGridRemoteControlPropertyData(TArray(OutBytes)));
 	return true;
 }
 
-bool URenderGridJob::SetRemoteControlValue(const TSharedPtr<FRemoteControlEntity>& RemoteControlEntity, const TArray<uint8>& BinaryArray)
+bool URenderGridJob::SetRemoteControlValueBytes(const FGuid& FieldId, const TArray<uint8>& Bytes)
 {
-	if (!RemoteControlEntity.IsValid())
+	if (!FieldId.IsValid())
 	{
 		return false;
 	}
-	const FString Key = RemoteControlEntity->GetId().ToString();
-	RemoteControlValues.Add(Key, FRenderGridRemoteControlPropertyData(TArray(BinaryArray)));
+
+	RemoteControlValues.Add(FieldId, FRenderGridRemoteControlPropertyData(TArray(Bytes)));
 	return true;
 }
 
+bool URenderGridJob::HasRemoteControlValue(const FGuid& FieldId) const
+{
+	return HasRemoteControlValueBytes(FieldId);
+}
+
+bool URenderGridJob::ConstGetRemoteControlValue(const FGuid& FieldId, FString& OutJson) const
+{
+	OutJson.Empty();
+	TArray<uint8> Bytes;
+	if (!ConstGetRemoteControlValueBytes(FieldId, Bytes))
+	{
+		return false;
+	}
+	OutJson.Append(UE::RenderGrid::Private::FRenderGridUtils::GetRemoteControlValueJsonFromBytes(Bytes));
+	return true;
+}
+
+bool URenderGridJob::GetRemoteControlValue(const FGuid& FieldId, FString& OutJson)
+{
+	OutJson.Empty();
+	TArray<uint8> Bytes;
+	if (!GetRemoteControlValueBytes(FieldId, Bytes))
+	{
+		return false;
+	}
+	OutJson.Append(UE::RenderGrid::Private::FRenderGridUtils::GetRemoteControlValueJsonFromBytes(Bytes));
+	return true;
+}
+
+bool URenderGridJob::SetRemoteControlValue(const FGuid& FieldId, const FString& Json)
+{
+	TArray<uint8> Bytes;
+	if (!GetRemoteControlValueBytes(FieldId, Bytes))
+	{
+		return false;
+	}
+	return SetRemoteControlValueBytes(FieldId, UE::RenderGrid::Private::FRenderGridUtils::GetRemoteControlValueBytesFromJson(Bytes, Json));
+}
+
+bool URenderGridJob::GetRemoteControlFieldIdFromLabel(const FString& Label, FGuid& OutFieldId)
+{
+	OutFieldId.Invalidate();
+	const FName LabelName = FName(Label);
+	for (URemoteControlPreset* RemoteControlPreset : GetRemoteControlPresets())
+	{
+		if (const FGuid FieldId = RemoteControlPreset->GetExposedEntityId(LabelName); FieldId.IsValid())
+		{
+			OutFieldId = FieldId;
+			return true;
+		}
+	}
+	return false;
+}
+
+bool URenderGridJob::GetRemoteControlLabelFromFieldId(const FGuid& FieldId, FString& OutLabel)
+{
+	OutLabel.Empty();
+	for (URemoteControlPreset* RemoteControlPreset : GetRemoteControlPresets())
+	{
+		for (const TWeakPtr<FRemoteControlEntity>& PropWeakPtr : RemoteControlPreset->GetExposedEntities<FRemoteControlEntity>())
+		{
+			if (const TSharedPtr<FRemoteControlEntity> Prop = PropWeakPtr.Pin())
+			{
+				if (Prop->GetId() == FieldId)
+				{
+					OutLabel = Prop->GetLabel().ToString();
+					return true;
+				}
+			}
+		}
+	}
+	return false;
+}
+
+TMap<FGuid, FString> URenderGridJob::GetRemoteControlValues()
+{
+	TMap<FGuid, FString> Result;
+	for (TTuple<FGuid, FRenderGridRemoteControlPropertyData> Entry : RemoteControlValues)
+	{
+		FString Json;
+		if (GetRemoteControlValue(Entry.Key, Json))
+		{
+			Result.Add(Entry.Key, Json);
+		}
+	}
+	return Result;
+}
+
 
 URenderGrid::URenderGrid()
 	: Guid(FGuid::NewGuid())
@@ -945,12 +1130,15 @@ URenderGridJob* URenderGrid::CreateTempRenderGridJob()
 
 	if (URenderGridPropsSourceRemoteControl* PropsSource = GetPropsSource<URenderGridPropsSourceRemoteControl>())
 	{
-		TArray<uint8> BinaryArray;
+		TArray<uint8> Bytes;
 		for (URenderGridPropRemoteControl* Field : PropsSource->GetProps()->GetAllCasted())
 		{
-			if (Field->GetValue(BinaryArray))
+			if (Field->GetValue(Bytes))
 			{
-				Job->SetRemoteControlValue(Field->GetRemoteControlEntity(), BinaryArray);
+				if (const TSharedPtr<FRemoteControlEntity> FieldEntity = Field->GetRemoteControlEntity(); FieldEntity.IsValid())
+				{
+					Job->SetRemoteControlValueBytes(FieldEntity, Bytes);
+				}
 			}
 		}
 	}
diff --git a/Engine/Plugins/Experimental/RenderGrid/Source/RenderGrid/Private/RenderGrid/RenderGridManager.cpp b/Engine/Plugins/Experimental/RenderGrid/Source/RenderGrid/Private/RenderGrid/RenderGridManager.cpp
--- a/Engine/Plugins/Experimental/RenderGrid/Source/RenderGrid/Private/RenderGrid/RenderGridManager.cpp
+++ b/Engine/Plugins/Experimental/RenderGrid/Source/RenderGrid/Private/RenderGrid/RenderGridManager.cpp
@@ -4,7 +4,6 @@
 #include "RenderGrid/RenderGridQueue.h"
 #include "RenderGridUtils.h"
 #include "MoviePipelineHighResSetting.h"
-#include "MoviePipelinePIEExecutor.h"
 
 
 URenderGridQueue* UE::RenderGrid::FRenderGridManager::CreateBatchRenderQueue(URenderGrid* Grid, const TArray<URenderGridJob*>& Jobs)
@@ -190,19 +189,21 @@ void UE::RenderGrid::FRenderGridManager::UpdateRenderGridJobsPropValues(URenderG
 	}
 
 	TArray<URenderGridJob*> Jobs = Grid->GetRenderGridJobs();
-	TArray<uint8> BinaryArray;
+	TArray<uint8> Bytes;
 	for (URenderGridPropRemoteControl* Field : PropsSource->GetProps()->GetAllCasted())
 	{
-		if (!Field->GetValue(BinaryArray))
+		if (!Field->GetValue(Bytes))
 		{
 			continue;
 		}
-		const TSharedPtr<FRemoteControlEntity> Entity = Field->GetRemoteControlEntity();
-		for (URenderGridJob* Job : Jobs)
+		if (const TSharedPtr<FRemoteControlEntity> Entity = Field->GetRemoteControlEntity(); Entity.IsValid())
 		{
-			if (!Job->HasRemoteControlValue(Entity))
+			for (URenderGridJob* Job : Jobs)
 			{
-				Job->SetRemoteControlValue(Entity, BinaryArray);
+				if (!Job->HasRemoteControlValueBytes(Entity))
+				{
+					Job->SetRemoteControlValueBytes(Entity, Bytes);
+				}
 			}
 		}
 	}
@@ -229,7 +230,7 @@ FRenderGridManagerPreviousPropValues UE::RenderGrid::FRenderGridManager::ApplyJo
 			}
 
 			TArray<uint8> PropData;
-			if (!Job->ConstGetRemoteControlValue(Prop->GetRemoteControlEntity(), PropData))
+			if (!Job->ConstGetRemoteControlValueBytes(Prop->GetRemoteControlEntity(), PropData))
 			{
 				continue;
 			}
diff --git a/Engine/Plugins/Experimental/RenderGrid/Source/RenderGrid/Private/RenderGrid/RenderGridPropsSource.cpp b/Engine/Plugins/Experimental/RenderGrid/Source/RenderGrid/Private/RenderGrid/RenderGridPropsSource.cpp
--- a/Engine/Plugins/Experimental/RenderGrid/Source/RenderGrid/Private/RenderGrid/RenderGridPropsSource.cpp
+++ b/Engine/Plugins/Experimental/RenderGrid/Source/RenderGrid/Private/RenderGrid/RenderGridPropsSource.cpp
@@ -3,6 +3,7 @@
 #include "RenderGrid/RenderGridPropsSource.h"
 #include "RenderGridLog.h"
 #include "IRemoteControlModule.h"
+#include "RCVirtualProperty.h"
 #include "Backends/JsonStructDeserializerBackend.h"
 #include "Backends/JsonStructSerializerBackend.h"
 
@@ -37,9 +38,9 @@ bool /*URenderGridPropRemoteControl::*/GetObjectRef(const TSharedPtr<FRemoteCont
 	return false;
 }
 
-bool URenderGridPropRemoteControl::GetValueOfEntity(const TSharedPtr<FRemoteControlEntity>& RemoteControlEntity, TArray<uint8>& OutBinaryArray)
+bool URenderGridPropRemoteControl::GetValueOfEntity(const TSharedPtr<FRemoteControlEntity>& RemoteControlEntity, TArray<uint8>& OutBytes)
 {
-	OutBinaryArray.Empty();
+	OutBytes.Empty();
 
 	TSharedPtr<FRemoteControlProperty> Field = StaticCastSharedPtr<FRemoteControlProperty>(RemoteControlEntity);
 	if (!Field.IsValid())
@@ -52,12 +53,12 @@ bool URenderGridPropRemoteControl::GetValueOfEntity(const TSharedPtr<FRemoteCont
 	{
 		return false;
 	}
-	FMemoryWriter Writer = FMemoryWriter(OutBinaryArray);
+	FMemoryWriter Writer = FMemoryWriter(OutBytes);
 	FJsonStructSerializerBackend WriterBackend = FJsonStructSerializerBackend(Writer, EStructSerializerBackendFlags::Default);
 	return IRemoteControlModule::Get().GetObjectProperties(ObjectRef, WriterBackend);
 }
 
-bool URenderGridPropRemoteControl::SetValueOfEntity(const TSharedPtr<FRemoteControlEntity>& RemoteControlEntity, const TArray<uint8>& BinaryArray)
+bool URenderGridPropRemoteControl::SetValueOfEntity(const TSharedPtr<FRemoteControlEntity>& RemoteControlEntity, const TArray<uint8>& Bytes)
 {
 	TSharedPtr<FRemoteControlProperty> Field = StaticCastSharedPtr<FRemoteControlProperty>(RemoteControlEntity);
 	if (!Field.IsValid())
@@ -68,12 +69,12 @@ bool URenderGridPropRemoteControl::SetValueOfEntity(const TSharedPtr<FRemoteCont
 	FRCObjectReference ObjectRefRead;
 	if (GetObjectRef(Field, ERCAccess::READ_ACCESS, ObjectRefRead))
 	{
-		TArray<uint8> CurrentBinaryArray;
-		FMemoryWriter Writer = FMemoryWriter(CurrentBinaryArray);
+		TArray<uint8> CurrentBytes;
+		FMemoryWriter Writer = FMemoryWriter(CurrentBytes);
 		FJsonStructSerializerBackend WriterBackend = FJsonStructSerializerBackend(Writer, EStructSerializerBackendFlags::Default);
 		if (IRemoteControlModule::Get().GetObjectProperties(ObjectRefRead, WriterBackend))
 		{
-			if (CurrentBinaryArray == BinaryArray)
+			if (CurrentBytes == Bytes)
 			{
 				// if the given value is already set, don't do anything
 				return true;
@@ -86,7 +87,7 @@ bool URenderGridPropRemoteControl::SetValueOfEntity(const TSharedPtr<FRemoteCont
 	{
 		return false;
 	}
-	FMemoryReader Reader = FMemoryReader(BinaryArray);
+	FMemoryReader Reader = FMemoryReader(Bytes);
 	FJsonStructDeserializerBackend ReaderBackend = FJsonStructDeserializerBackend(Reader);
 	if (!IRemoteControlModule::Get().SetObjectProperties(ObjectRefWrite, ReaderBackend, ERCPayloadType::Json))
 	{
@@ -119,7 +120,7 @@ bool URenderGridPropRemoteControl::SetValueOfEntity(const TSharedPtr<FRemoteCont
 	return true;
 }
 
-bool URenderGridPropRemoteControl::CanSetValueOfEntity(const TSharedPtr<FRemoteControlEntity>& RemoteControlEntity, const TArray<uint8>& BinaryArray)
+bool URenderGridPropRemoteControl::CanSetValueOfEntity(const TSharedPtr<FRemoteControlEntity>& RemoteControlEntity, const TArray<uint8>& Bytes)
 {
 	TSharedPtr<FRemoteControlProperty> Field = StaticCastSharedPtr<FRemoteControlProperty>(RemoteControlEntity);
 	if (!Field.IsValid())
@@ -130,12 +131,12 @@ bool URenderGridPropRemoteControl::CanSetValueOfEntity(const TSharedPtr<FRemoteC
 	FRCObjectReference ObjectRefRead;
 	if (GetObjectRef(Field, ERCAccess::READ_ACCESS, ObjectRefRead))
 	{
-		TArray<uint8> CurrentBinaryArray;
-		FMemoryWriter Writer = FMemoryWriter(CurrentBinaryArray);
+		TArray<uint8> CurrentBytes;
+		FMemoryWriter Writer = FMemoryWriter(CurrentBytes);
 		FJsonStructSerializerBackend WriterBackend = FJsonStructSerializerBackend(Writer, EStructSerializerBackendFlags::Default);
 		if (IRemoteControlModule::Get().GetObjectProperties(ObjectRefRead, WriterBackend))
 		{
-			if (CurrentBinaryArray == BinaryArray)
+			if (CurrentBytes == Bytes)
 			{
 				// if the given value is already set, don't do anything
 				return true;
@@ -198,6 +199,46 @@ TArray<URenderGridPropRemoteControl*> URenderGridPropsRemoteControl::GetAllCaste
 	return Result;
 }
 
+URenderGridPropRemoteControl* URenderGridPropsRemoteControl::GetById(const FGuid& Id) const
+{
+	if (IsValid(RemoteControlPreset))
+	{
+		for (const TWeakPtr<FRemoteControlEntity>& PropWeakPtr : RemoteControlPreset->GetExposedEntities<FRemoteControlEntity>())
+		{
+			if (const TSharedPtr<FRemoteControlEntity> Prop = PropWeakPtr.Pin())
+			{
+				if (Prop->GetId() == Id)
+				{
+					URenderGridPropRemoteControl* PropObj = NewObject<URenderGridPropRemoteControl>(const_cast<URenderGridPropsRemoteControl*>(this));
+					PropObj->Initialize(Prop);
+					return PropObj;
+				}
+			}
+		}
+	}
+	return nullptr;
+}
+
+URenderGridPropRemoteControl* URenderGridPropsRemoteControl::GetByLabel(const FString& Label) const
+{
+	if (IsValid(RemoteControlPreset))
+	{
+		for (const TWeakPtr<FRemoteControlEntity>& PropWeakPtr : RemoteControlPreset->GetExposedEntities<FRemoteControlEntity>())
+		{
+			if (const TSharedPtr<FRemoteControlEntity> Prop = PropWeakPtr.Pin())
+			{
+				if (Prop->GetLabel().ToString().TrimStartAndEnd().Equals(Label.TrimStartAndEnd(), ESearchCase::CaseSensitive))
+				{
+					URenderGridPropRemoteControl* PropObj = NewObject<URenderGridPropRemoteControl>(const_cast<URenderGridPropsRemoteControl*>(this));
+					PropObj->Initialize(Prop);
+					return PropObj;
+				}
+			}
+		}
+	}
+	return nullptr;
+}
+
 
 void URenderGridPropsSourceRemoteControl::SetSourceOrigin(UObject* SourceOrigin)
 {
diff --git a/Engine/Plugins/Experimental/RenderGrid/Source/RenderGrid/Private/RenderGridUtils.cpp b/Engine/Plugins/Experimental/RenderGrid/Source/RenderGrid/Private/RenderGridUtils.cpp
--- a/Engine/Plugins/Experimental/RenderGrid/Source/RenderGrid/Private/RenderGridUtils.cpp
+++ b/Engine/Plugins/Experimental/RenderGrid/Source/RenderGrid/Private/RenderGridUtils.cpp
@@ -7,12 +7,14 @@
 #include "HAL/PlatformFileManager.h"
 #include "IImageWrapper.h"
 #include "IImageWrapperModule.h"
-#include "ImageUtils.h"
 #include "TextureResource.h"
 #include "Editor/EditorPerformanceSettings.h"
 #include "Engine/Engine.h"
 #include "Misc/Paths.h"
 #include "Modules/ModuleManager.h"
+#include "Policies/CondensedJsonPrintPolicy.h"
+#include "Serialization/JsonReader.h"
+#include "Serialization/JsonSerializer.h"
 
 
 bool UE::RenderGrid::Private::FRenderGridUtils::IsImage(const FString& ImagePath)
@@ -309,6 +311,109 @@ FString UE::RenderGrid::Private::FRenderGridUtils::DenormalizeJobOutputDirectory
 }
 
 
+FString UE::RenderGrid::Private::FRenderGridUtils::GetRemoteControlValueJsonFromBytes(const TArray<uint8>& ValueBytes, const bool PrettyPrint)
+{
+	TSharedPtr<FJsonValue> Value;
+	{
+		const FUTF16ToTCHAR ConvertedString(reinterpret_cast<UTF16CHAR*>(const_cast<uint8*>(ValueBytes.GetData())), ValueBytes.Num() / sizeof(UTF16CHAR));
+		const FString Json = FString(ConvertedString.Length(), ConvertedString.Get()).TrimStartAndEnd();
+
+		const TSharedRef<TJsonReader<TCHAR>> JsonReader = TJsonReaderFactory<TCHAR>::Create(Json);
+		TSharedPtr<FJsonObject> JsonObjectData;
+		if (!FJsonSerializer::Deserialize(JsonReader, JsonObjectData))
+		{
+			// the given bytes isn't valid json
+			return TEXT("NULL");
+		}
+
+		TArray<TSharedPtr<FJsonValue>> Values;
+		JsonObjectData->Values.GenerateValueArray(Values);
+		if (Values.Num() != 1)
+		{
+			// in the given bytes (json data), there is no key, or there are multiple keys, either way it's invalid
+			return TEXT("NULL");
+		}
+
+		Value = Values[0];
+	}
+
+	FString JsonString;
+	{
+		if (PrettyPrint)
+		{
+			TSharedRef<TJsonWriter<TCHAR, TPrettyJsonPrintPolicy<TCHAR>>> JsonWriter = TJsonWriterFactory<TCHAR, TPrettyJsonPrintPolicy<TCHAR>>::Create(&JsonString);
+			FJsonSerializer::Serialize({Value}, JsonWriter);
+			JsonWriter->Close();
+		}
+		else
+		{
+			TSharedRef<TJsonWriter<TCHAR, TCondensedJsonPrintPolicy<TCHAR>>> JsonWriter = TJsonWriterFactory<TCHAR, TCondensedJsonPrintPolicy<TCHAR>>::Create(&JsonString);
+			FJsonSerializer::Serialize({Value}, JsonWriter);
+			JsonWriter->Close();
+		}
+
+		JsonString.TrimStartAndEndInline();
+		JsonString.MidInline(1, JsonString.Len() - 2);
+		JsonString.TrimStartAndEndInline();
+	}
+
+	return JsonString;
+}
+
+TArray<uint8> UE::RenderGrid::Private::FRenderGridUtils::GetRemoteControlValueBytesFromJson(const TArray<uint8>& OldValueBytes, const FString& NewValueJson)
+{
+	FString OldValueKey;
+	{
+		const FUTF16ToTCHAR ConvertedString(reinterpret_cast<UTF16CHAR*>(const_cast<uint8*>(OldValueBytes.GetData())), OldValueBytes.Num() / sizeof(UTF16CHAR));
+		const FString Json = FString(ConvertedString.Length(), ConvertedString.Get());
+
+		const TSharedRef<TJsonReader<TCHAR>> JsonReader = TJsonReaderFactory<TCHAR>::Create(Json);
+		TSharedPtr<FJsonObject> JsonObjectData;
+		if (!FJsonSerializer::Deserialize(JsonReader, JsonObjectData))
+		{
+			// the given bytes isn't valid json
+			return TArray<uint8>();
+		}
+
+		TArray<FString> Keys;
+		JsonObjectData->Values.GenerateKeyArray(Keys);
+		if (Keys.Num() != 1)
+		{
+			// in the given bytes (json data), there is no key, or there are multiple keys, either way it's invalid
+			return TArray<uint8>();
+		}
+
+		OldValueKey = Keys[0];
+	}
+
+	TSharedPtr<FJsonValue> Value;
+	{
+		const TSharedRef<TJsonReader<TCHAR>> JsonReader = TJsonReaderFactory<TCHAR>::Create(TEXT("{\"value\":") + NewValueJson + TEXT("}"));
+		TSharedPtr<FJsonObject> JsonObjectData;
+		if (!FJsonSerializer::Deserialize(JsonReader, JsonObjectData))
+		{
+			// the given new value (json string) isn't valid json
+			return TArray<uint8>();
+		}
+
+		Value = JsonObjectData->TryGetField(TEXT("value"));
+	}
+
+	FString JsonString;
+	{
+		TSharedRef<FJsonObject> JsonObject = MakeShareable(new FJsonObject());
+		JsonObject->SetField(OldValueKey, Value);
+
+		TSharedRef<TJsonWriter<TCHAR, TCondensedJsonPrintPolicy<TCHAR>>> JsonWriter = TJsonWriterFactory<TCHAR, TCondensedJsonPrintPolicy<TCHAR>>::Create(&JsonString);
+		FJsonSerializer::Serialize(JsonObject, JsonWriter);
+		JsonWriter->Close();
+	}
+
+	FTCHARToUTF16 UTF16String(*JsonString, JsonString.Len());
+	return TArray(reinterpret_cast<const uint8*>(UTF16String.Get()), UTF16String.Length() * sizeof(UTF16CHAR));
+}
+
+
 FRenderGridPreviousEngineFpsSettings UE::RenderGrid::Private::FRenderGridUtils::DisableFpsLimit()
 {
 	FRenderGridPreviousEngineFpsSettings Settings;
diff --git a/Engine/Plugins/Experimental/RenderGrid/Source/RenderGrid/Public/RenderGrid/RenderGrid.h b/Engine/Plugins/Experimental/RenderGrid/Source/RenderGrid/Public/RenderGrid/RenderGrid.h
--- a/Engine/Plugins/Experimental/RenderGrid/Source/RenderGrid/Public/RenderGrid/RenderGrid.h
+++ b/Engine/Plugins/Experimental/RenderGrid/Source/RenderGrid/Public/RenderGrid/RenderGrid.h
@@ -264,11 +264,40 @@ public:
 	void SetRenderPreset(UMoviePipelinePrimaryConfig* NewRenderPreset) { RenderPreset = NewRenderPreset; }
 
 
-	bool HasRemoteControlValue(const TSharedPtr<FRemoteControlEntity>& RemoteControlEntity) const;
-	bool ConstGetRemoteControlValue(const TSharedPtr<FRemoteControlEntity>& RemoteControlEntity, TArray<uint8>& OutBinaryArray) const;
-	bool GetRemoteControlValue(const TSharedPtr<FRemoteControlEntity>& RemoteControlEntity, TArray<uint8>& OutBinaryArray);
-	bool SetRemoteControlValue(const TSharedPtr<FRemoteControlEntity>& RemoteControlEntity, const TArray<uint8>& BinaryArray);
-	TMap<FString, FRenderGridRemoteControlPropertyData>& GetRemoteControlValuesRef() { return RemoteControlValues; }
+	TArray<URemoteControlPreset*> GetRemoteControlPresets();
+
+	bool HasRemoteControlValueBytes(const TSharedPtr<FRemoteControlEntity>& RemoteControlEntity) const;
+	bool ConstGetRemoteControlValueBytes(const TSharedPtr<FRemoteControlEntity>& RemoteControlEntity, TArray<uint8>& OutBytes) const;
+	bool GetRemoteControlValueBytes(const TSharedPtr<FRemoteControlEntity>& RemoteControlEntity, TArray<uint8>& OutBytes);
+	bool SetRemoteControlValueBytes(const TSharedPtr<FRemoteControlEntity>& RemoteControlEntity, const TArray<uint8>& Bytes);
+
+	bool HasRemoteControlValueBytes(const FGuid& FieldId) const;
+	bool ConstGetRemoteControlValueBytes(const FGuid& FieldId, TArray<uint8>& OutBytes) const;
+	bool GetRemoteControlValueBytes(const FGuid& FieldId, TArray<uint8>& OutBytes);
+	bool SetRemoteControlValueBytes(const FGuid& FieldId, const TArray<uint8>& Bytes);
+
+	TMap<FGuid, FRenderGridRemoteControlPropertyData>& GetRemoteControlValuesBytesRef() { return RemoteControlValues; }
+
+
+	UFUNCTION(BlueprintPure, Category="Render Grid Job")
+	bool HasRemoteControlValue(const FGuid& FieldId) const;
+
+	bool ConstGetRemoteControlValue(const FGuid& FieldId, FString& OutJson) const;
+
+	UFUNCTION(BlueprintPure, Category="Render Grid Job")
+	bool GetRemoteControlValue(const FGuid& FieldId, FString& OutJson);
+
+	UFUNCTION(BlueprintCallable, Category="Render Grid Job")
+	bool SetRemoteControlValue(const FGuid& FieldId, const FString& Json);
+
+	UFUNCTION(BlueprintPure, Category="Render Grid Job")
+	bool GetRemoteControlFieldIdFromLabel(const FString& Label, FGuid& OutFieldId);
+
+	UFUNCTION(BlueprintPure, Category="Render Grid Job")
+	bool GetRemoteControlLabelFromFieldId(const FGuid& FieldId, FString& OutLabel);
+
+	UFUNCTION(BlueprintCallable, Category="Render Grid Job")
+	TMap<FGuid, FString> GetRemoteControlValues();
 
 private:
 	/** The unique ID of this job. */
@@ -329,7 +358,7 @@ private:
 
 	/** The remote control plugin can be used to customize and modify the way a job is rendered. If remote control is being used, the property values of this job will be stored in this map (remote control entity id -> value as bytes). */
 	UPROPERTY()
-	TMap<FString, FRenderGridRemoteControlPropertyData> RemoteControlValues;
+	TMap<FGuid, FRenderGridRemoteControlPropertyData> RemoteControlValues;
 };
 
 
diff --git a/Engine/Plugins/Experimental/RenderGrid/Source/RenderGrid/Public/RenderGrid/RenderGridPropsSource.h b/Engine/Plugins/Experimental/RenderGrid/Source/RenderGrid/Public/RenderGrid/RenderGridPropsSource.h
--- a/Engine/Plugins/Experimental/RenderGrid/Source/RenderGrid/Public/RenderGrid/RenderGridPropsSource.h
+++ b/Engine/Plugins/Experimental/RenderGrid/Source/RenderGrid/Public/RenderGrid/RenderGridPropsSource.h
@@ -128,13 +128,13 @@ class RENDERGRID_API URenderGridPropRemoteControl : public URenderGridPropBase
 
 public:
 	/** Gets the value (as bytes) of the given property (remote control entity). Returns true if the operation was successful, false otherwise. */
-	static bool GetValueOfEntity(const TSharedPtr<FRemoteControlEntity>& RemoteControlEntity, TArray<uint8>& OutBinaryArray);
+	static bool GetValueOfEntity(const TSharedPtr<FRemoteControlEntity>& RemoteControlEntity, TArray<uint8>& OutBytes);
 
 	/** Set the value (as bytes) of the given property (remote control entity). Returns true if the operation was successful, false otherwise. */
-	static bool SetValueOfEntity(const TSharedPtr<FRemoteControlEntity>& RemoteControlEntity, const TArray<uint8>& BinaryArray);
+	static bool SetValueOfEntity(const TSharedPtr<FRemoteControlEntity>& RemoteControlEntity, const TArray<uint8>& Bytes);
 
 	/** Tests if it can set the value (as bytes) of the given property (remote control entity). Returns true if the set operation would likely be successful, false otherwise. */
-	static bool CanSetValueOfEntity(const TSharedPtr<FRemoteControlEntity>& RemoteControlEntity, const TArray<uint8>& BinaryArray);
+	static bool CanSetValueOfEntity(const TSharedPtr<FRemoteControlEntity>& RemoteControlEntity, const TArray<uint8>& Bytes);
 
 public:
 	/** Sets the initial values of this instance. */
@@ -144,13 +144,13 @@ public:
 	TSharedPtr<FRemoteControlEntity> GetRemoteControlEntity() const { return RemoteControlEntity; }
 
 	/** Gets the value (as bytes) of this property. Returns true if the operation was successful, false otherwise. */
-	bool GetValue(TArray<uint8>& OutBinaryArray) const { return GetValueOfEntity(RemoteControlEntity, OutBinaryArray); }
+	bool GetValue(TArray<uint8>& OutBytes) const { return GetValueOfEntity(RemoteControlEntity, OutBytes); }
 
 	/** Set the value (as bytes) of this property. Returns true if the operation was successful, false otherwise. */
-	bool SetValue(const TArray<uint8>& BinaryArray) { return SetValueOfEntity(RemoteControlEntity, BinaryArray); }
+	bool SetValue(const TArray<uint8>& Bytes) { return SetValueOfEntity(RemoteControlEntity, Bytes); }
 
 	/** Tests if it can set the value (as bytes) of this property. Returns true if the set operation would likely be successful, false otherwise. */
-	bool CanSetValue(const TArray<uint8>& BinaryArray) { return CanSetValueOfEntity(RemoteControlEntity, BinaryArray); }
+	bool CanSetValue(const TArray<uint8>& Bytes) { return CanSetValueOfEntity(RemoteControlEntity, Bytes); }
 
 protected:
 	/** The property, which is a remote control entity (which can be a field or a function). */
@@ -176,6 +176,12 @@ public:
 	/** Returns all props, casted to URenderGridPropRemoteControl, for ease of use. */
 	TArray<URenderGridPropRemoteControl*> GetAllCasted() const;
 
+	/** Returns the prop with the given ID, or NULL if not found. */
+	URenderGridPropRemoteControl* GetById(const FGuid& Id) const;
+
+	/** Returns the prop with the given label, or NULL if not found. */
+	URenderGridPropRemoteControl* GetByLabel(const FString& Label) const;
+
 	/** Returns the remote control preset. */
 	URemoteControlPreset* GetRemoteControlPreset() const { return (IsValid(RemoteControlPreset) ? RemoteControlPreset : nullptr); }
 
@@ -203,6 +209,9 @@ public:
 	/** Returns (in the out parameter) the preset groups that are available in this remote control preset. */
 	void GetAvailablePresetGroups(TArray<FName>& OutPresetGroups) const;
 
+	/** Returns the remote control preset. */
+	URemoteControlPreset* GetRemoteControlPreset() const { return (IsValid(RemoteControlPreset) ? RemoteControlPreset : nullptr); }
+
 protected:
 	/** The source of properties, which is a remote control preset. */
 	UPROPERTY(VisibleInstanceOnly, Category="Render Grid|Remote Control")
diff --git a/Engine/Plugins/Experimental/RenderGrid/Source/RenderGrid/Public/RenderGridUtils.h b/Engine/Plugins/Experimental/RenderGrid/Source/RenderGrid/Public/RenderGridUtils.h
--- a/Engine/Plugins/Experimental/RenderGrid/Source/RenderGrid/Public/RenderGridUtils.h
+++ b/Engine/Plugins/Experimental/RenderGrid/Source/RenderGrid/Public/RenderGridUtils.h
@@ -168,6 +168,17 @@ namespace UE::RenderGrid::Private
 		static FString DenormalizeJobOutputDirectory(const FString& NormalizedOutputDirectory);
 
 
+		/**
+		 * Converts the given bytes to a UTF-16 string, parses it as JSON, obtains the value from it, converts it to a string, and returns that string.
+		 */
+		static FString GetRemoteControlValueJsonFromBytes(const TArray<uint8>& ValueBytes, const bool PrettyPrint = false);
+		
+		/**
+		 * Converts the given bytes to a UTF-16 string, parses it as JSON, and obtains the key from it. Then it wraps the given json to insert the key (found in the given bytes), converts the json to bytes (UTF-16), and then return those bytes.
+		 */
+		static TArray<uint8> GetRemoteControlValueBytesFromJson(const TArray<uint8>& OldValueBytes, const FString& NewValueJson);
+
+
 		/**
 		 * Disables the current FPS limiting, returns the previous settings which allow you to revert back to the previous state.
 		 */
diff --git a/Engine/Plugins/Experimental/RenderGrid/Source/RenderGrid/RenderGrid.Build.cs b/Engine/Plugins/Experimental/RenderGrid/Source/RenderGrid/RenderGrid.Build.cs
--- a/Engine/Plugins/Experimental/RenderGrid/Source/RenderGrid/RenderGrid.Build.cs
+++ b/Engine/Plugins/Experimental/RenderGrid/Source/RenderGrid/RenderGrid.Build.cs
@@ -23,6 +23,7 @@ public class RenderGrid : ModuleRules
 		PrivateDependencyModuleNames.AddRange(
 			new string[]
 			{
+				"Json",
 				"LevelSequence",
 				"LevelSequenceEditor",
 				"MovieRenderPipelineCore",
diff --git a/Engine/Plugins/Experimental/RenderGrid/Source/RenderGridEditor/Private/UI/SRenderGridPropsRemoteControl.cpp b/Engine/Plugins/Experimental/RenderGrid/Source/RenderGridEditor/Private/UI/SRenderGridPropsRemoteControl.cpp
--- a/Engine/Plugins/Experimental/RenderGrid/Source/RenderGridEditor/Private/UI/SRenderGridPropsRemoteControl.cpp
+++ b/Engine/Plugins/Experimental/RenderGrid/Source/RenderGridEditor/Private/UI/SRenderGridPropsRemoteControl.cpp
@@ -159,22 +159,22 @@ void UE::RenderGrid::Private::SRenderGridPropsRemoteControl::OnRemoteControlExpo
 			{
 				if (const TSharedPtr<FRemoteControlEntity> Entity = Preset->GetExposedEntity<FRemoteControlEntity>(Id).Pin())
 				{
-					TArray<uint8> BinaryArray;
-					if (!URenderGridPropRemoteControl::GetValueOfEntity(Entity, BinaryArray))
+					TArray<uint8> Bytes;
+					if (!URenderGridPropRemoteControl::GetValueOfEntity(Entity, Bytes))
 					{
 						continue;
 					}
-					TArray<uint8> StoredBinaryArray;
-					if (!GetSelectedJobFieldValue(Entity, StoredBinaryArray))
+					TArray<uint8> StoredBytes;
+					if (!GetSelectedJobFieldValue(Entity, StoredBytes))
 					{
 						continue;
 					}
-					if (BinaryArray == StoredBinaryArray)
+					if (Bytes == StoredBytes)
 					{
 						continue;
 					}
 
-					if (!SetSelectedJobFieldValue(Entity, BinaryArray))
+					if (!SetSelectedJobFieldValue(Entity, Bytes))
 					{
 						continue;
 					}
@@ -205,23 +205,29 @@ URenderGridJob* UE::RenderGrid::Private::SRenderGridPropsRemoteControl::GetSelec
 	return nullptr;
 }
 
-bool UE::RenderGrid::Private::SRenderGridPropsRemoteControl::GetSelectedJobFieldValue(const TSharedPtr<FRemoteControlEntity>& RemoteControlEntity, TArray<uint8>& OutBinaryArray)
+bool UE::RenderGrid::Private::SRenderGridPropsRemoteControl::GetSelectedJobFieldValue(const TSharedPtr<FRemoteControlEntity>& RemoteControlEntity, TArray<uint8>& OutBytes)
 {
-	OutBinaryArray.Empty();
-	if (URenderGridJob* SelectedJob = GetSelectedJob(); IsValid(SelectedJob))
+	OutBytes.Empty();
+	if (RemoteControlEntity.IsValid())
 	{
-		return SelectedJob->GetRemoteControlValue(RemoteControlEntity, OutBinaryArray);
+		if (URenderGridJob* SelectedJob = GetSelectedJob(); IsValid(SelectedJob))
+		{
+			return SelectedJob->GetRemoteControlValueBytes(RemoteControlEntity, OutBytes);
+		}
 	}
 	return false;
 }
 
-bool UE::RenderGrid::Private::SRenderGridPropsRemoteControl::SetSelectedJobFieldValue(const TSharedPtr<FRemoteControlEntity>& RemoteControlEntity, const TArray<uint8>& BinaryArray)
+bool UE::RenderGrid::Private::SRenderGridPropsRemoteControl::SetSelectedJobFieldValue(const TSharedPtr<FRemoteControlEntity>& RemoteControlEntity, const TArray<uint8>& Bytes)
 {
-	if (URenderGridJob* SelectedJob = GetSelectedJob(); IsValid(SelectedJob))
+	if (RemoteControlEntity.IsValid())
 	{
-		FScopedTransaction Transaction(LOCTEXT("ChangeJobProperty", "Change Job Property"));
-		SelectedJob->Modify();
-		return SelectedJob->SetRemoteControlValue(RemoteControlEntity, BinaryArray);
+		if (URenderGridJob* SelectedJob = GetSelectedJob(); IsValid(SelectedJob))
+		{
+			FScopedTransaction Transaction(LOCTEXT("ChangeJobProperty", "Change Job Property"));
+			SelectedJob->Modify();
+			return SelectedJob->SetRemoteControlValueBytes(RemoteControlEntity, Bytes);
+		}
 	}
 	return false;
 }
diff --git a/Engine/Plugins/Experimental/RenderGrid/Source/RenderGridEditor/Private/UI/SRenderGridPropsRemoteControl.h b/Engine/Plugins/Experimental/RenderGrid/Source/RenderGridEditor/Private/UI/SRenderGridPropsRemoteControl.h
--- a/Engine/Plugins/Experimental/RenderGrid/Source/RenderGridEditor/Private/UI/SRenderGridPropsRemoteControl.h
+++ b/Engine/Plugins/Experimental/RenderGrid/Source/RenderGridEditor/Private/UI/SRenderGridPropsRemoteControl.h
@@ -49,10 +49,10 @@ namespace UE::RenderGrid::Private
 		URenderGridJob* GetSelectedJob();
 
 		/** Obtains the value (as bytes) of the given prop (the given remote control entity), returns true if it succeeded, returns false otherwise. */
-		bool GetSelectedJobFieldValue(const TSharedPtr<FRemoteControlEntity>& RemoteControlEntity, TArray<uint8>& OutBinaryArray);
+		bool GetSelectedJobFieldValue(const TSharedPtr<FRemoteControlEntity>& RemoteControlEntity, TArray<uint8>& OutBytes);
 
 		/** Sets the value of the given prop (the given remote control entity) with the given value (as bytes), returns true if it succeeded, returns false otherwise. */
-		bool SetSelectedJobFieldValue(const TSharedPtr<FRemoteControlEntity>& RemoteControlEntity, const TArray<uint8>& BinaryArray);
+		bool SetSelectedJobFieldValue(const TSharedPtr<FRemoteControlEntity>& RemoteControlEntity, const TArray<uint8>& Bytes);
 
 	private:
 		/** A reference to the blueprint editor that owns the render grid instance. */
