RenderGrid plugin:
- Added blueprints to create and parse FProperty JSON data with. This allows for easy access (getting and setting) to the remote control property values.
- Improved the documentation of the current python scripts.

diff --git a/Engine/Plugins/Experimental/RenderGrid/Source/RenderGrid/Private/RenderGrid/RenderGrid.cpp b/Engine/Plugins/Experimental/RenderGrid/Source/RenderGrid/Private/RenderGrid/RenderGrid.cpp
--- a/Engine/Plugins/Experimental/RenderGrid/Source/RenderGrid/Private/RenderGrid/RenderGrid.cpp
+++ b/Engine/Plugins/Experimental/RenderGrid/Source/RenderGrid/Private/RenderGrid/RenderGrid.cpp
@@ -10,7 +10,7 @@
 #include "MovieScene.h"
 #include "RenderGrid/RenderGridManager.h"
 #include "RenderGrid/RenderGridQueue.h"
-#include "RenderGridUtils.h"
+#include "Utils/RenderGridUtils.h"
 #include "UObject/ObjectSaveContext.h"
 
 
@@ -546,27 +546,27 @@ bool URenderGridJob::HasRemoteControlValue(const FGuid& FieldId) const
 	return HasRemoteControlValueBytes(FieldId);
 }
 
-bool URenderGridJob::ConstGetRemoteControlValue(const FGuid& FieldId, FString& OutJson) const
+bool URenderGridJob::ConstGetRemoteControlValue(const FGuid& FieldId, FString& Json) const
 {
-	OutJson.Empty();
+	Json.Empty();
 	TArray<uint8> Bytes;
 	if (!ConstGetRemoteControlValueBytes(FieldId, Bytes))
 	{
 		return false;
 	}
-	OutJson.Append(UE::RenderGrid::Private::FRenderGridUtils::GetRemoteControlValueJsonFromBytes(Bytes));
+	Json.Append(UE::RenderGrid::Private::FRenderGridUtils::GetRemoteControlValueJsonFromBytes(Bytes));
 	return true;
 }
 
-bool URenderGridJob::GetRemoteControlValue(const FGuid& FieldId, FString& OutJson)
+bool URenderGridJob::GetRemoteControlValue(const FGuid& FieldId, FString& Json)
 {
-	OutJson.Empty();
+	Json.Empty();
 	TArray<uint8> Bytes;
 	if (!GetRemoteControlValueBytes(FieldId, Bytes))
 	{
 		return false;
 	}
-	OutJson.Append(UE::RenderGrid::Private::FRenderGridUtils::GetRemoteControlValueJsonFromBytes(Bytes));
+	Json.Append(UE::RenderGrid::Private::FRenderGridUtils::GetRemoteControlValueJsonFromBytes(Bytes));
 	return true;
 }
 
@@ -580,24 +580,24 @@ bool URenderGridJob::SetRemoteControlValue(const FGuid& FieldId, const FString&
 	return SetRemoteControlValueBytes(FieldId, UE::RenderGrid::Private::FRenderGridUtils::GetRemoteControlValueBytesFromJson(Bytes, Json));
 }
 
-bool URenderGridJob::GetRemoteControlFieldIdFromLabel(const FString& Label, FGuid& OutFieldId)
+bool URenderGridJob::GetRemoteControlFieldIdFromLabel(const FString& Label, FGuid& FieldId)
 {
-	OutFieldId.Invalidate();
+	FieldId.Invalidate();
 	const FName LabelName = FName(Label);
 	for (URemoteControlPreset* RemoteControlPreset : GetRemoteControlPresets())
 	{
-		if (const FGuid FieldId = RemoteControlPreset->GetExposedEntityId(LabelName); FieldId.IsValid())
+		if (const FGuid RemoteControlFieldId = RemoteControlPreset->GetExposedEntityId(LabelName); FieldId.IsValid())
 		{
-			OutFieldId = FieldId;
+			FieldId = RemoteControlFieldId;
 			return true;
 		}
 	}
 	return false;
 }
 
-bool URenderGridJob::GetRemoteControlLabelFromFieldId(const FGuid& FieldId, FString& OutLabel)
+bool URenderGridJob::GetRemoteControlLabelFromFieldId(const FGuid& FieldId, FString& Label)
 {
-	OutLabel.Empty();
+	Label.Empty();
 	for (URemoteControlPreset* RemoteControlPreset : GetRemoteControlPresets())
 	{
 		for (const TWeakPtr<FRemoteControlEntity>& PropWeakPtr : RemoteControlPreset->GetExposedEntities<FRemoteControlEntity>())
@@ -606,7 +606,7 @@ bool URenderGridJob::GetRemoteControlLabelFromFieldId(const FGuid& FieldId, FStr
 			{
 				if (Prop->GetId() == FieldId)
 				{
-					OutLabel = Prop->GetLabel().ToString();
+					Label = Prop->GetLabel().ToString();
 					return true;
 				}
 			}
diff --git a/Engine/Plugins/Experimental/RenderGrid/Source/RenderGrid/Private/RenderGrid/RenderGridManager.cpp b/Engine/Plugins/Experimental/RenderGrid/Source/RenderGrid/Private/RenderGrid/RenderGridManager.cpp
--- a/Engine/Plugins/Experimental/RenderGrid/Source/RenderGrid/Private/RenderGrid/RenderGridManager.cpp
+++ b/Engine/Plugins/Experimental/RenderGrid/Source/RenderGrid/Private/RenderGrid/RenderGridManager.cpp
@@ -2,7 +2,7 @@
 
 #include "RenderGrid/RenderGridManager.h"
 #include "RenderGrid/RenderGridQueue.h"
-#include "RenderGridUtils.h"
+#include "Utils/RenderGridUtils.h"
 #include "MoviePipelineHighResSetting.h"
 
 
diff --git a/Engine/Plugins/Experimental/RenderGrid/Source/RenderGrid/Private/RenderGrid/RenderGridQueue.cpp b/Engine/Plugins/Experimental/RenderGrid/Source/RenderGrid/Private/RenderGrid/RenderGridQueue.cpp
--- a/Engine/Plugins/Experimental/RenderGrid/Source/RenderGrid/Private/RenderGrid/RenderGridQueue.cpp
+++ b/Engine/Plugins/Experimental/RenderGrid/Source/RenderGrid/Private/RenderGrid/RenderGridQueue.cpp
@@ -5,7 +5,7 @@
 #include "IRenderGridModule.h"
 #include "RenderGrid/RenderGrid.h"
 #include "RenderGrid/RenderGridManager.h"
-#include "RenderGridUtils.h"
+#include "Utils/RenderGridUtils.h"
 #include "UObject/Package.h"
 #include "Utils/RenderGridGenericExecutionQueue.h"
 
diff --git a/Engine/Plugins/Experimental/RenderGrid/Source/RenderGrid/Private/RenderGridUtils.cpp b/Engine/Plugins/Experimental/RenderGrid/Source/RenderGrid/Private/RenderGridUtils.cpp
deleted file mode 100644
diff --git a/Engine/Plugins/Experimental/RenderGrid/Source/RenderGrid/Private/Utils/RenderGridGenericExecutionQueue.h b/Engine/Plugins/Experimental/RenderGrid/Source/RenderGrid/Private/Utils/RenderGridGenericExecutionQueue.h
deleted file mode 100644
diff --git a/Engine/Plugins/Experimental/RenderGrid/Source/RenderGrid/Private/Utils/RenderGridRemoteControlUtils.cpp b/Engine/Plugins/Experimental/RenderGrid/Source/RenderGrid/Private/Utils/RenderGridRemoteControlUtils.cpp
new file mode 100644
--- /dev/null
+++ b/Engine/Plugins/Experimental/RenderGrid/Source/RenderGrid/Private/Utils/RenderGridRemoteControlUtils.cpp
@@ -0,0 +1,373 @@
+﻿// Copyright Epic Games, Inc. All Rights Reserved.
+
+#include "Utils/RenderGridRemoteControlUtils.h"
+
+#include "InstancedStruct.h"
+#include "StructDeserializer.h"
+#include "StructSerializer.h"
+#include "Backends/JsonStructDeserializerBackend.h"
+#include "Backends/JsonStructSerializerBackend.h"
+#include "Utils/RenderGridUtils.h"
+#include "Dom/JsonValue.h"
+#include "Engine/Engine.h"
+#include "Serialization/MemoryReader.h"
+
+
+URenderGridRemoteControlUtils::URenderGridRemoteControlUtils(const FObjectInitializer& ObjectInitializer) : Super(ObjectInitializer) {}
+
+
+template<typename NumericType>
+NumericType CastDoubleToNumericType(const double Value)
+{
+	if constexpr (std::is_same_v<NumericType, double>)
+	{
+		return Value;
+	}
+	else
+	{
+		if (Value < TNumericLimits<NumericType>::Lowest())
+		{
+			return TNumericLimits<NumericType>::Lowest();
+		}
+		if (Value > TNumericLimits<NumericType>::Max())
+		{
+			return TNumericLimits<NumericType>::Max();
+		}
+		return static_cast<NumericType>(Value);
+	}
+}
+
+
+template<typename NumericType>
+void ParseJsonAsNumericType(const FString& Json, const NumericType DefaultValue, NumericType& Value, bool& bSuccess)
+{
+	TSharedPtr<FJsonValue> JsonValue = UE::RenderGrid::Private::FRenderGridUtils::ParseJson(Json);
+	if (!JsonValue.IsValid())
+	{
+		bSuccess = false;
+		Value = DefaultValue;
+		return;
+	}
+	bSuccess = true;
+	Value = CastDoubleToNumericType<NumericType>(JsonValue->AsNumber());
+}
+
+void ParseJsonAsStructType(const FString& Json, void* Value, UStruct& ValueTypeInfo, bool& bSuccess)
+{
+	FTCHARToUTF16 UTF16String(*Json, Json.Len());
+	TArray<uint8> Bytes = TArray(reinterpret_cast<const uint8*>(UTF16String.Get()), UTF16String.Length() * sizeof(UTF16CHAR));
+
+	FMemoryReader Reader = FMemoryReader(Bytes);
+	FJsonStructDeserializerBackend ReaderBackend = FJsonStructDeserializerBackend(Reader);
+
+	bSuccess = FStructDeserializer::Deserialize(Value, ValueTypeInfo, ReaderBackend);
+	Reader.Close();
+}
+
+void ParseJsonAsStructType(const FString& Json, const FInstancedStruct& DefaultValue, FInstancedStruct& Value, bool& bSuccess)
+{
+	if (!DefaultValue.IsValid())
+	{
+		bSuccess = false;
+		Value = DefaultValue;
+		return;
+	}
+	TSharedPtr<FStructOnScope> StructOnScope = MakeShareable(new FStructOnScope(DefaultValue.GetScriptStruct()));
+	UStruct* ValueTypeInfo = const_cast<UStruct*>(StructOnScope->GetStruct());
+	ParseJsonAsStructType(Json, StructOnScope->GetStructMemory(), *ValueTypeInfo, bSuccess);
+	if (!bSuccess)
+	{
+		Value = DefaultValue;
+		return;
+	}
+	Value = FInstancedStruct();
+	Value.InitializeAs(DefaultValue.GetScriptStruct(), StructOnScope->GetStructMemory());
+}
+
+template<typename StructType>
+void ParseJsonAsStructType(const FString& Json, const StructType& DefaultValue, StructType& Value, bool& bSuccess)
+{
+	FInstancedStruct ValueWrapped = FInstancedStruct::Make<StructType>(Value);
+	ParseJsonAsStructType(Json, FInstancedStruct::Make<StructType>(DefaultValue), ValueWrapped, bSuccess);
+	Value = ValueWrapped.Get<StructType>();
+}
+
+void URenderGridRemoteControlUtils::ParseJsonAsByte(const FString& Json, const uint8 DefaultValue, bool& bSuccess, uint8& Value)
+{
+	ParseJsonAsNumericType<uint8>(Json, DefaultValue, Value, bSuccess);
+}
+
+void URenderGridRemoteControlUtils::ParseJsonAsInt32(const FString& Json, const int32 DefaultValue, bool& bSuccess, int32& Value)
+{
+	ParseJsonAsNumericType<int32>(Json, DefaultValue, Value, bSuccess);
+}
+
+void URenderGridRemoteControlUtils::ParseJsonAsInt64(const FString& Json, const int64 DefaultValue, bool& bSuccess, int64& Value)
+{
+	ParseJsonAsNumericType<int64>(Json, DefaultValue, Value, bSuccess);
+}
+
+void URenderGridRemoteControlUtils::ParseJsonAsFloat(const FString& Json, const float DefaultValue, bool& bSuccess, float& Value)
+{
+	ParseJsonAsNumericType<float>(Json, DefaultValue, Value, bSuccess);
+}
+
+void URenderGridRemoteControlUtils::ParseJsonAsDouble(const FString& Json, const double DefaultValue, bool& bSuccess, double& Value)
+{
+	ParseJsonAsNumericType<double>(Json, DefaultValue, Value, bSuccess);
+}
+
+void URenderGridRemoteControlUtils::ParseJsonAsBoolean(const FString& Json, const bool DefaultValue, bool& bSuccess, bool& Value)
+{
+	TSharedPtr<FJsonValue> JsonValue = UE::RenderGrid::Private::FRenderGridUtils::ParseJson(Json);
+	if (!JsonValue.IsValid())
+	{
+		bSuccess = false;
+		Value = DefaultValue;
+		return;
+	}
+	bSuccess = true;
+	Value = JsonValue->AsBool();
+}
+
+void URenderGridRemoteControlUtils::ParseJsonAsString(const FString& Json, const FString& DefaultValue, bool& bSuccess, FString& Value)
+{
+	TSharedPtr<FJsonValue> JsonValue = UE::RenderGrid::Private::FRenderGridUtils::ParseJson(Json);
+	if (!JsonValue.IsValid())
+	{
+		bSuccess = false;
+		Value = DefaultValue;
+		return;
+	}
+	bSuccess = true;
+	Value = JsonValue->AsString();
+}
+
+void URenderGridRemoteControlUtils::ParseJsonAsName(const FString& Json, const FName& DefaultValue, bool& bSuccess, FName& Value)
+{
+	FString ValueString;
+	ParseJsonAsString(Json, DefaultValue.ToString(), bSuccess, ValueString);
+	Value = FName(ValueString);
+}
+
+void URenderGridRemoteControlUtils::ParseJsonAsText(const FString& Json, const FText& DefaultValue, bool& bSuccess, FText& Value)
+{
+	TSharedPtr<FJsonValue> JsonValue = UE::RenderGrid::Private::FRenderGridUtils::ParseJson(Json);
+	if (!JsonValue.IsValid())
+	{
+		bSuccess = false;
+		Value = DefaultValue;
+		return;
+	}
+	FText TextValue;
+	if (!FTextStringHelper::ReadFromBuffer(*JsonValue->AsString(), TextValue))
+	{
+		bSuccess = false;
+		Value = DefaultValue;
+		return;
+	}
+	bSuccess = true;
+	Value = TextValue;
+}
+
+void URenderGridRemoteControlUtils::ParseJsonAsObjectReference(const FString& Json, UObject* DefaultValue, bool& bSuccess, UObject*& Value)
+{
+	FString ObjectPath;
+	ParseJsonAsString(Json, TEXT(""), bSuccess, ObjectPath);
+	if (!bSuccess || ObjectPath.IsEmpty())
+	{
+		bSuccess = false;
+		Value = DefaultValue;
+		return;
+	}
+	FSoftObjectPath SoftObjectPath = FSoftObjectPath(ObjectPath);
+	if (!SoftObjectPath.IsValid())
+	{
+		bSuccess = false;
+		Value = DefaultValue;
+		return;
+	}
+	Value = TSoftObjectPtr(SoftObjectPath).LoadSynchronous();
+	if (!IsValid(Value))
+	{
+		bSuccess = false;
+		Value = DefaultValue;
+		return;
+	}
+	bSuccess = true;
+}
+
+void URenderGridRemoteControlUtils::ParseJsonAsClassReference(const FString& Json, UClass* DefaultValue, bool& bSuccess, UClass*& Value)
+{
+	FString ClassPath;
+	ParseJsonAsString(Json, TEXT(""), bSuccess, ClassPath);
+	if (!bSuccess || ClassPath.IsEmpty())
+	{
+		bSuccess = false;
+		Value = DefaultValue;
+		return;
+	}
+	FSoftClassPath SoftClassPath = FSoftClassPath(ClassPath);
+	if (!SoftClassPath.IsValid())
+	{
+		bSuccess = false;
+		Value = DefaultValue;
+		return;
+	}
+	Value = TSoftClassPtr(SoftClassPath).LoadSynchronous();
+	if (!IsValid(Value))
+	{
+		bSuccess = false;
+		Value = DefaultValue;
+		return;
+	}
+	bSuccess = true;
+}
+
+void URenderGridRemoteControlUtils::ParseJsonAsStruct(const FString& Json, const FInstancedStruct& DefaultValue, bool& bSuccess, FInstancedStruct& Value)
+{
+	ParseJsonAsStructType(Json, DefaultValue, Value, bSuccess);
+}
+
+void URenderGridRemoteControlUtils::ParseJsonAsVector(const FString& Json, const FVector& DefaultValue, bool& bSuccess, FVector& Value)
+{
+	ParseJsonAsStructType<FVector>(Json, DefaultValue, Value, bSuccess);
+}
+
+void URenderGridRemoteControlUtils::ParseJsonAsRotator(const FString& Json, const FRotator& DefaultValue, bool& bSuccess, FRotator& Value)
+{
+	ParseJsonAsStructType<FRotator>(Json, DefaultValue, Value, bSuccess);
+}
+
+void URenderGridRemoteControlUtils::ParseJsonAsTransform(const FString& Json, const FTransform& DefaultValue, bool& bSuccess, FTransform& Value)
+{
+	ParseJsonAsStructType<FTransform>(Json, DefaultValue, Value, bSuccess);
+}
+
+void URenderGridRemoteControlUtils::ParseJsonAsColor(const FString& Json, const FColor& DefaultValue, bool& bSuccess, FColor& Value)
+{
+	ParseJsonAsStructType<FColor>(Json, DefaultValue, Value, bSuccess);
+}
+
+void URenderGridRemoteControlUtils::ParseJsonAsLinearColor(const FString& Json, const FLinearColor& DefaultValue, bool& bSuccess, FLinearColor& Value)
+{
+	ParseJsonAsStructType<FLinearColor>(Json, DefaultValue, Value, bSuccess);
+}
+
+
+void StructTypeToJson(const void* Value, UStruct& ValueTypeInfo, FString& Json)
+{
+	TArray<uint8> Bytes;
+	FMemoryWriter Writer = FMemoryWriter(Bytes);
+	FJsonStructSerializerBackend WriterBackend = FJsonStructSerializerBackend(Writer, EStructSerializerBackendFlags::Default);
+
+	FStructSerializer::Serialize(Value, ValueTypeInfo, WriterBackend);
+	Writer.Close();
+
+	const FUTF16ToTCHAR ConvertedString(reinterpret_cast<UTF16CHAR*>(Bytes.GetData()), Bytes.Num() / sizeof(UTF16CHAR));
+	Json = FString(ConvertedString.Length(), ConvertedString.Get()).TrimStartAndEnd();
+}
+
+void StructTypeToJson(const FInstancedStruct& Value, FString& Json)
+{
+	if (!Value.IsValid())
+	{
+		Json = TEXT("");
+		return;
+	}
+	UStruct* ValueTypeInfo = const_cast<UScriptStruct*>(Value.GetScriptStruct());
+	StructTypeToJson(Value.GetMemory(), *ValueTypeInfo, Json);
+}
+
+template<typename StructType>
+void StructTypeToJson(const StructType& Value, FString& Json)
+{
+	StructTypeToJson(FInstancedStruct::Make<StructType>(Value), Json);
+}
+
+void URenderGridRemoteControlUtils::ByteToJson(const uint8 Value, FString& Json)
+{
+	Json = UE::RenderGrid::Private::FRenderGridUtils::ToJsonString(MakeShareable(new FJsonValueNumber(Value)));
+}
+
+void URenderGridRemoteControlUtils::Int32ToJson(const int32 Value, FString& Json)
+{
+	Json = UE::RenderGrid::Private::FRenderGridUtils::ToJsonString(MakeShareable(new FJsonValueNumber(Value)));
+}
+
+void URenderGridRemoteControlUtils::Int64ToJson(const int64 Value, FString& Json)
+{
+	Json = UE::RenderGrid::Private::FRenderGridUtils::ToJsonString(MakeShareable(new FJsonValueNumber(Value)));
+}
+
+void URenderGridRemoteControlUtils::FloatToJson(const float Value, FString& Json)
+{
+	Json = UE::RenderGrid::Private::FRenderGridUtils::ToJsonString(MakeShareable(new FJsonValueNumber(Value)));
+}
+
+void URenderGridRemoteControlUtils::DoubleToJson(const double Value, FString& Json)
+{
+	Json = UE::RenderGrid::Private::FRenderGridUtils::ToJsonString(MakeShareable(new FJsonValueNumber(Value)));
+}
+
+void URenderGridRemoteControlUtils::BooleanToJson(const bool Value, FString& Json)
+{
+	Json = UE::RenderGrid::Private::FRenderGridUtils::ToJsonString(MakeShareable(new FJsonValueBoolean(Value)));
+}
+
+void URenderGridRemoteControlUtils::StringToJson(const FString& Value, FString& Json)
+{
+	Json = UE::RenderGrid::Private::FRenderGridUtils::ToJsonString(MakeShareable(new FJsonValueString(Value)));
+}
+
+void URenderGridRemoteControlUtils::NameToJson(const FName& Value, FString& Json)
+{
+	StringToJson(Value.ToString(), Json);
+}
+
+void URenderGridRemoteControlUtils::TextToJson(const FText& Value, FString& Json)
+{
+	FString ValueAsString;
+	FTextStringHelper::WriteToBuffer(ValueAsString, Value);
+	StringToJson(ValueAsString, Json);
+}
+
+void URenderGridRemoteControlUtils::ObjectReferenceToJson(UObject* Value, FString& Json)
+{
+	StringToJson(GetPathNameSafe(Value), Json);
+}
+
+void URenderGridRemoteControlUtils::ClassReferenceToJson(UClass* Value, FString& Json)
+{
+	StringToJson(GetPathNameSafe(Value), Json);
+}
+
+void URenderGridRemoteControlUtils::StructToJson(const FInstancedStruct& Value, FString& Json)
+{
+	StructTypeToJson(Value, Json);
+}
+
+void URenderGridRemoteControlUtils::VectorToJson(const FVector& Value, FString& Json)
+{
+	StructTypeToJson<FVector>(Value, Json);
+}
+
+void URenderGridRemoteControlUtils::RotatorToJson(const FRotator& Value, FString& Json)
+{
+	StructTypeToJson<FRotator>(Value, Json);
+}
+
+void URenderGridRemoteControlUtils::TransformToJson(const FTransform& Value, FString& Json)
+{
+	StructTypeToJson<FTransform>(Value, Json);
+}
+
+void URenderGridRemoteControlUtils::ColorToJson(const FColor& Value, FString& Json)
+{
+	StructTypeToJson<FColor>(Value, Json);
+}
+
+void URenderGridRemoteControlUtils::LinearColorToJson(const FLinearColor& Value, FString& Json)
+{
+	StructTypeToJson<FLinearColor>(Value, Json);
+}
diff --git a/Engine/Plugins/Experimental/RenderGrid/Source/RenderGrid/Private/Utils/RenderGridUtils.cpp b/Engine/Plugins/Experimental/RenderGrid/Source/RenderGrid/Private/Utils/RenderGridUtils.cpp
new file mode 100644
--- /dev/null
+++ b/Engine/Plugins/Experimental/RenderGrid/Source/RenderGrid/Private/Utils/RenderGridUtils.cpp
@@ -0,0 +1,470 @@
+// Copyright Epic Games, Inc. All Rights Reserved.
+
+#include "Utils/RenderGridUtils.h"
+#include "RenderGrid/RenderGrid.h"
+#include "Engine/Texture2D.h"
+#include "HAL/FileManager.h"
+#include "HAL/PlatformFileManager.h"
+#include "IImageWrapper.h"
+#include "IImageWrapperModule.h"
+#include "TextureResource.h"
+#include "Editor/EditorPerformanceSettings.h"
+#include "Engine/Engine.h"
+#include "Misc/Paths.h"
+#include "Modules/ModuleManager.h"
+#include "Policies/CondensedJsonPrintPolicy.h"
+#include "Serialization/JsonReader.h"
+#include "Serialization/JsonSerializer.h"
+
+
+bool UE::RenderGrid::Private::FRenderGridUtils::IsImage(const FString& ImagePath)
+{
+	static IImageWrapperModule* ImageWrapperModule = FModuleManager::LoadModulePtr<IImageWrapperModule>("ImageWrapper");
+	if (ImageWrapperModule == nullptr)
+	{
+		return false;
+	}
+
+	EImageFormat ImageFormat = ImageWrapperModule->GetImageFormatFromExtension(*ImagePath);
+	if (ImageFormat == EImageFormat::Invalid)
+	{
+		return false;
+	}
+
+	return true;
+}
+
+UTexture2D* UE::RenderGrid::Private::FRenderGridUtils::GetImage(const FString& ImagePath, UTexture2D* Texture2D, bool& bOutReusedGivenTexture2D)
+{
+	static IImageWrapperModule* ImageWrapperModule = FModuleManager::LoadModulePtr<IImageWrapperModule>("ImageWrapper");
+	if (ImageWrapperModule == nullptr)
+	{
+		bOutReusedGivenTexture2D = (Texture2D == nullptr);
+		return nullptr;
+	}
+
+	EImageFormat ImageFormat = ImageWrapperModule->GetImageFormatFromExtension(*ImagePath);
+	if (ImageFormat == EImageFormat::Invalid)
+	{
+		bOutReusedGivenTexture2D = (Texture2D == nullptr);
+		return nullptr;
+	}
+
+	TArray<uint8> ImageBytes = GetFileData(ImagePath);
+	if (ImageBytes.Num() <= 0)
+	{
+		bOutReusedGivenTexture2D = (Texture2D == nullptr);
+		return nullptr;
+	}
+
+	return BytesToExistingImage(ImageBytes, ImageFormat, Texture2D, bOutReusedGivenTexture2D);
+}
+
+
+UTexture2D* UE::RenderGrid::Private::FRenderGridUtils::BytesToImage(const TArray<uint8>& ByteArray, const EImageFormat ImageFormat)
+{
+	static IImageWrapperModule* ImageWrapperModule = FModuleManager::LoadModulePtr<IImageWrapperModule>("ImageWrapper");
+	if (ImageWrapperModule == nullptr)
+	{
+		return nullptr;
+	}
+
+	TSharedPtr<IImageWrapper> ImageWrapper = ImageWrapperModule->CreateImageWrapper(ImageFormat);
+	if (!ImageWrapper.IsValid() || !ImageWrapper->SetCompressed(ByteArray.GetData(), ByteArray.Num()))
+	{
+		return nullptr;
+	}
+
+	TArray<uint8> Uncompressed;
+	if (!ImageWrapper->GetRaw(ERGBFormat::BGRA, 8, Uncompressed))
+	{
+		return nullptr;
+	}
+
+	return DataToTexture2D(ImageWrapper->GetWidth(), ImageWrapper->GetHeight(), Uncompressed.GetData(), Uncompressed.Num());
+}
+
+UTexture2D* UE::RenderGrid::Private::FRenderGridUtils::BytesToExistingImage(const TArray<uint8>& ByteArray, const EImageFormat ImageFormat, UTexture2D* Texture2D, bool& bOutReusedGivenTexture2D)
+{
+	static IImageWrapperModule* ImageWrapperModule = FModuleManager::LoadModulePtr<IImageWrapperModule>("ImageWrapper");
+	if (ImageWrapperModule == nullptr)
+	{
+		bOutReusedGivenTexture2D = (Texture2D == nullptr);
+		return nullptr;
+	}
+
+	TSharedPtr<IImageWrapper> ImageWrapper = ImageWrapperModule->CreateImageWrapper(ImageFormat);
+	if (!ImageWrapper.IsValid() || !ImageWrapper->SetCompressed(ByteArray.GetData(), ByteArray.Num()))
+	{
+		bOutReusedGivenTexture2D = (Texture2D == nullptr);
+		return nullptr;
+	}
+
+	TArray<uint8> Uncompressed;
+	if (!ImageWrapper->GetRaw(ERGBFormat::BGRA, 8, Uncompressed))
+	{
+		bOutReusedGivenTexture2D = (Texture2D == nullptr);
+		return nullptr;
+	}
+
+	UTexture2D* NewTexture2D = DataToExistingTexture2D(ImageWrapper->GetWidth(), ImageWrapper->GetHeight(), Uncompressed.GetData(), Uncompressed.Num(), Texture2D, bOutReusedGivenTexture2D);
+	bOutReusedGivenTexture2D = (Texture2D == NewTexture2D);
+	return NewTexture2D;
+}
+
+
+UTexture2D* UE::RenderGrid::Private::FRenderGridUtils::DataToTexture2D(int32 Width, int32 Height, const void* Src, SIZE_T Count)
+{
+	UTexture2D* Texture2D = UTexture2D::CreateTransient(Width, Height, PF_B8G8R8A8);
+	if (Texture2D == nullptr)
+	{
+		return nullptr;
+	}
+	Texture2D->bNoTiling = true;
+
+#if WITH_EDITORONLY_DATA
+	Texture2D->MipGenSettings = TMGS_NoMipmaps;
+#endif
+
+	void* TextureData = Texture2D->GetPlatformData()->Mips[0].BulkData.Lock(LOCK_READ_WRITE);
+	FMemory::Memcpy(TextureData, Src, Count);
+	Texture2D->GetPlatformData()->Mips[0].BulkData.Unlock();
+
+	Texture2D->UpdateResource();
+	return Texture2D;
+}
+
+UTexture2D* UE::RenderGrid::Private::FRenderGridUtils::DataToExistingTexture2D(int32 Width, int32 Height, const void* Src, SIZE_T Count, UTexture2D* Texture2D, bool& bOutReusedGivenTexture2D)
+{
+	UTexture2D* OriginalTexture2D = Texture2D;
+
+	bOutReusedGivenTexture2D = true;
+	if ((Texture2D == nullptr) || (Texture2D->GetSizeX() != Width) || (Texture2D->GetSizeY() != Height) || (Texture2D->GetPixelFormat() != PF_B8G8R8A8))
+	{
+		bOutReusedGivenTexture2D = false;
+		Texture2D = UTexture2D::CreateTransient(Width, Height, PF_B8G8R8A8);
+	}
+	if (Texture2D == nullptr)
+	{
+		bOutReusedGivenTexture2D = (OriginalTexture2D == nullptr);
+		return nullptr;
+	}
+	Texture2D->bNoTiling = true;
+
+#if WITH_EDITORONLY_DATA
+	Texture2D->MipGenSettings = TMGS_NoMipmaps;
+#endif
+
+	void* TextureData = Texture2D->GetPlatformData()->Mips[0].BulkData.Lock(LOCK_READ_WRITE);
+	FMemory::Memcpy(TextureData, Src, Count);
+	Texture2D->GetPlatformData()->Mips[0].BulkData.Unlock();
+
+	Texture2D->UpdateResource();
+	return Texture2D;
+}
+
+
+TArray<FString> UE::RenderGrid::Private::FRenderGridUtils::GetFiles(const FString& Directory, const bool bRecursive)
+{
+	class FRenderGridUtilsGetFilesVisitor : public IPlatformFile::FDirectoryVisitor
+	{
+	public:
+		TArray<FString> Files;
+		virtual bool Visit(const TCHAR* FileName, const bool bIsDirectory) override
+		{
+			if (!bIsDirectory)
+			{
+				Files.Add(FileName);
+			}
+			return true;
+		}
+	};
+	FRenderGridUtilsGetFilesVisitor GetFilesVisitor;
+	if (bRecursive)
+	{
+		FPlatformFileManager::Get().GetPlatformFile().IterateDirectoryRecursively(*Directory, GetFilesVisitor);
+	}
+	else
+	{
+		FPlatformFileManager::Get().GetPlatformFile().IterateDirectory(*Directory, GetFilesVisitor);
+	}
+	return GetFilesVisitor.Files;
+}
+
+TArray<uint8> UE::RenderGrid::Private::FRenderGridUtils::GetFileData(const FString& File)
+{
+	TArray<uint8> Data;
+	IFileHandle* FileHandle = FPlatformFileManager::Get().GetPlatformFile().OpenRead(*File);
+	if (FileHandle != nullptr)
+	{
+		int64 FileSize = FileHandle->Size();
+		if (FileSize > 0x7fffffff)
+		{
+			FileSize = 0;
+		}
+		if (FileSize > 0)
+		{
+			Data.SetNum(FileSize);
+			FileHandle->Read(Data.GetData(), FileSize);
+		}
+		delete FileHandle;
+	}
+	return Data;
+}
+
+
+void UE::RenderGrid::Private::FRenderGridUtils::DeleteDirectory(const FString& Directory)
+{
+	IFileManager::Get().DeleteDirectory(*Directory, false, true);
+}
+
+void UE::RenderGrid::Private::FRenderGridUtils::EmptyDirectory(const FString& Directory)
+{
+	class FRenderGridUtilsEmptyDirectoryVisitor : public IPlatformFile::FDirectoryVisitor
+	{
+	public:
+		virtual bool Visit(const TCHAR* FileName, const bool bIsDirectory) override
+		{
+			if (bIsDirectory)
+			{
+				IFileManager::Get().DeleteDirectory(FileName, false, false);
+			}
+			else
+			{
+				IFileManager::Get().Delete(FileName, false, true, true);
+			}
+			return true;
+		}
+	};
+	FRenderGridUtilsEmptyDirectoryVisitor GetFilesVisitor;
+	FPlatformFileManager::Get().GetPlatformFile().IterateDirectoryRecursively(*Directory, GetFilesVisitor);
+}
+
+
+FString UE::RenderGrid::Private::FRenderGridUtils::NormalizeOutputDirectory(const FString& OutputDirectory)
+{
+	FString NewOutputDirectory = OutputDirectory;
+	FPaths::NormalizeDirectoryName(NewOutputDirectory);
+	FPaths::RemoveDuplicateSlashes(NewOutputDirectory);
+	FPaths::CollapseRelativeDirectories(NewOutputDirectory);
+	FPaths::RemoveDuplicateSlashes(NewOutputDirectory);
+	if (!NewOutputDirectory.EndsWith(TEXT("/"), ESearchCase::CaseSensitive))
+	{
+		NewOutputDirectory += TEXT("/");
+	}
+	return NewOutputDirectory;
+}
+
+
+FString UE::RenderGrid::Private::FRenderGridUtils::PurgeJobIdOrReturnEmptyString(const FString& JobId)
+{
+	static FString ValidCharacters = TEXT("abcdefghijklmnopqrstuvwxyzABCDEFGHIJKLMNOPQRSTUVWXYZ0123456789-_");
+
+	FString Result;
+	for (TCHAR JobIdChar : JobId)
+	{
+		int32 Index;
+		if (ValidCharacters.FindChar(JobIdChar, Index))
+		{
+			Result += JobIdChar;
+		}
+	}
+	return Result;
+}
+
+FString UE::RenderGrid::Private::FRenderGridUtils::PurgeJobId(const FString& JobId)
+{
+	FString Result = PurgeJobIdOrReturnEmptyString(JobId);
+	if (Result.IsEmpty())
+	{
+		return TEXT("0");
+	}
+	return Result;
+}
+
+FString UE::RenderGrid::Private::FRenderGridUtils::PurgeJobIdOrGenerateUniqueId(URenderGrid* Grid, const FString& JobId)
+{
+	FString Result = PurgeJobIdOrReturnEmptyString(JobId);
+	if (Result.IsEmpty())
+	{
+		return Grid->GenerateNextJobId();
+	}
+	return Result;
+}
+
+FString UE::RenderGrid::Private::FRenderGridUtils::PurgeJobName(const FString& JobName)
+{
+	return JobName.TrimStartAndEnd();
+}
+
+FString UE::RenderGrid::Private::FRenderGridUtils::NormalizeJobOutputDirectory(const FString& OutputDirectory)
+{
+	return NormalizeOutputDirectory(FPaths::ConvertRelativePathToFull(OutputDirectory))
+		.Replace(*NormalizeOutputDirectory(FPaths::ConvertRelativePathToFull(FPaths::ProjectDir())), TEXT("{project_dir}/"));
+}
+
+FString UE::RenderGrid::Private::FRenderGridUtils::DenormalizeJobOutputDirectory(const FString& NormalizedOutputDirectory)
+{
+	return NormalizeOutputDirectory(FPaths::ConvertRelativePathToFull(
+		NormalizedOutputDirectory.Replace(TEXT("{project_dir}"), *FPaths::ProjectDir())
+	));
+}
+
+
+FString UE::RenderGrid::Private::FRenderGridUtils::ToJsonString(const TSharedPtr<FJsonValue>& Value, const bool PrettyPrint)
+{
+	if (!Value.IsValid())
+	{
+		return TEXT("NULL");
+	}
+
+	FString JsonString;
+	if (PrettyPrint)
+	{
+		TSharedRef<TJsonWriter<TCHAR, TPrettyJsonPrintPolicy<TCHAR>>> JsonWriter = TJsonWriterFactory<TCHAR, TPrettyJsonPrintPolicy<TCHAR>>::Create(&JsonString);
+		FJsonSerializer::Serialize({Value}, JsonWriter);
+		JsonWriter->Close();
+	}
+	else
+	{
+		TSharedRef<TJsonWriter<TCHAR, TCondensedJsonPrintPolicy<TCHAR>>> JsonWriter = TJsonWriterFactory<TCHAR, TCondensedJsonPrintPolicy<TCHAR>>::Create(&JsonString);
+		FJsonSerializer::Serialize({Value}, JsonWriter);
+		JsonWriter->Close();
+	}
+
+	JsonString.TrimStartAndEndInline();
+	JsonString.MidInline(1, JsonString.Len() - 2);
+	JsonString.TrimStartAndEndInline();
+	return JsonString;
+}
+
+TSharedPtr<FJsonValue> UE::RenderGrid::Private::FRenderGridUtils::ParseJson(const FString& Json)
+{
+	const TSharedRef<TJsonReader<TCHAR>> JsonReader = TJsonReaderFactory<TCHAR>::Create(TEXT("{\"value\":") + Json + TEXT("}"));
+	TSharedPtr<FJsonObject> JsonObjectData;
+	if (!FJsonSerializer::Deserialize(JsonReader, JsonObjectData))
+	{
+		// the given string isn't valid json
+		return TSharedPtr<FJsonValue>();
+	}
+	return JsonObjectData->TryGetField(TEXT("value"));
+}
+
+
+FString UE::RenderGrid::Private::FRenderGridUtils::GetRemoteControlValueJsonFromBytes(const TArray<uint8>& ValueBytes, const bool PrettyPrint)
+{
+	TSharedPtr<FJsonValue> Value;
+	{
+		const FUTF16ToTCHAR ConvertedString(reinterpret_cast<UTF16CHAR*>(const_cast<uint8*>(ValueBytes.GetData())), ValueBytes.Num() / sizeof(UTF16CHAR));
+		const FString Json = FString(ConvertedString.Length(), ConvertedString.Get()).TrimStartAndEnd();
+
+		const TSharedRef<TJsonReader<TCHAR>> JsonReader = TJsonReaderFactory<TCHAR>::Create(Json);
+		TSharedPtr<FJsonObject> JsonObjectData;
+		if (!FJsonSerializer::Deserialize(JsonReader, JsonObjectData))
+		{
+			// the given bytes isn't valid json
+			return TEXT("NULL");
+		}
+
+		TArray<TSharedPtr<FJsonValue>> Values;
+		JsonObjectData->Values.GenerateValueArray(Values);
+		if (Values.Num() != 1)
+		{
+			// in the given bytes (json data), there is no key, or there are multiple keys, either way it's invalid
+			return TEXT("NULL");
+		}
+
+		Value = Values[0];
+	}
+
+	return ToJsonString(Value, PrettyPrint);
+}
+
+TArray<uint8> UE::RenderGrid::Private::FRenderGridUtils::GetRemoteControlValueBytesFromJson(const TArray<uint8>& OldValueBytes, const FString& NewValueJson)
+{
+	FString OldValueKey;
+	{
+		const FUTF16ToTCHAR ConvertedString(reinterpret_cast<UTF16CHAR*>(const_cast<uint8*>(OldValueBytes.GetData())), OldValueBytes.Num() / sizeof(UTF16CHAR));
+		const FString Json = FString(ConvertedString.Length(), ConvertedString.Get());
+
+		const TSharedRef<TJsonReader<TCHAR>> JsonReader = TJsonReaderFactory<TCHAR>::Create(Json);
+		TSharedPtr<FJsonObject> JsonObjectData;
+		if (!FJsonSerializer::Deserialize(JsonReader, JsonObjectData))
+		{
+			// the given bytes isn't valid json
+			return TArray<uint8>();
+		}
+
+		TArray<FString> Keys;
+		JsonObjectData->Values.GenerateKeyArray(Keys);
+		if (Keys.Num() != 1)
+		{
+			// in the given bytes (json data), there is no key, or there are multiple keys, either way it's invalid
+			return TArray<uint8>();
+		}
+
+		OldValueKey = Keys[0];
+	}
+
+	TSharedPtr<FJsonValue> Value = ParseJson(NewValueJson);
+
+	FString JsonString;
+	{
+		TSharedRef<FJsonObject> JsonObject = MakeShareable(new FJsonObject());
+		JsonObject->SetField(OldValueKey, Value);
+
+		TSharedRef<TJsonWriter<TCHAR, TCondensedJsonPrintPolicy<TCHAR>>> JsonWriter = TJsonWriterFactory<TCHAR, TCondensedJsonPrintPolicy<TCHAR>>::Create(&JsonString);
+		FJsonSerializer::Serialize(JsonObject, JsonWriter);
+		JsonWriter->Close();
+	}
+
+	FTCHARToUTF16 UTF16String(*JsonString, JsonString.Len());
+	return TArray(reinterpret_cast<const uint8*>(UTF16String.Get()), UTF16String.Length() * sizeof(UTF16CHAR));
+}
+
+
+FRenderGridPreviousEngineFpsSettings UE::RenderGrid::Private::FRenderGridUtils::DisableFpsLimit()
+{
+	FRenderGridPreviousEngineFpsSettings Settings;
+	if (GEngine)
+	{
+		static IConsoleVariable* VSync = IConsoleManager::Get().FindConsoleVariable(TEXT("r.VSync"));
+		static IConsoleVariable* VSyncEditor = IConsoleManager::Get().FindConsoleVariable(TEXT("r.VSyncEditor"));
+		UEditorPerformanceSettings* EditorPerformanceSettings = GetMutableDefault<UEditorPerformanceSettings>();
+
+		Settings.bHasBeenSet = true;
+		Settings.bUseFixedFrameRate = GEngine->bUseFixedFrameRate;
+		Settings.bForceDisableFrameRateSmoothing = GEngine->bForceDisableFrameRateSmoothing;
+		Settings.MaxFps = GEngine->GetMaxFPS();
+		Settings.bVSync = VSync->GetBool();
+		Settings.bVSyncEditor = VSyncEditor->GetBool();
+		Settings.bThrottleCPUWhenNotForeground = EditorPerformanceSettings->bThrottleCPUWhenNotForeground;
+
+		GEngine->bUseFixedFrameRate = false;
+		GEngine->bForceDisableFrameRateSmoothing = true;
+		GEngine->SetMaxFPS(10000);
+		VSync->Set(false);
+		VSyncEditor->Set(false);
+		EditorPerformanceSettings->bThrottleCPUWhenNotForeground = false;
+		EditorPerformanceSettings->PostEditChange();
+	}
+	return Settings;
+}
+
+void UE::RenderGrid::Private::FRenderGridUtils::RestoreFpsLimit(const FRenderGridPreviousEngineFpsSettings& Settings)
+{
+	if (GEngine && Settings.bHasBeenSet)
+	{
+		static IConsoleVariable* VSync = IConsoleManager::Get().FindConsoleVariable(TEXT("r.VSync"));
+		static IConsoleVariable* VSyncEditor = IConsoleManager::Get().FindConsoleVariable(TEXT("r.VSyncEditor"));
+		UEditorPerformanceSettings* EditorPerformanceSettings = GetMutableDefault<UEditorPerformanceSettings>();
+
+		GEngine->bUseFixedFrameRate = Settings.bUseFixedFrameRate;
+		GEngine->bForceDisableFrameRateSmoothing = Settings.bForceDisableFrameRateSmoothing;
+		GEngine->SetMaxFPS(Settings.MaxFps);
+		VSync->Set(Settings.bVSync);
+		VSyncEditor->Set(Settings.bVSyncEditor);
+		EditorPerformanceSettings->bThrottleCPUWhenNotForeground = Settings.bThrottleCPUWhenNotForeground;
+		EditorPerformanceSettings->PostEditChange();
+	}
+}
diff --git a/Engine/Plugins/Experimental/RenderGrid/Source/RenderGrid/Public/RenderGrid/RenderGrid.h b/Engine/Plugins/Experimental/RenderGrid/Source/RenderGrid/Public/RenderGrid/RenderGrid.h
--- a/Engine/Plugins/Experimental/RenderGrid/Source/RenderGrid/Public/RenderGrid/RenderGrid.h
+++ b/Engine/Plugins/Experimental/RenderGrid/Source/RenderGrid/Public/RenderGrid/RenderGrid.h
@@ -282,19 +282,19 @@ public:
 	UFUNCTION(BlueprintPure, Category="Render Grid Job")
 	bool HasRemoteControlValue(const FGuid& FieldId) const;
 
-	bool ConstGetRemoteControlValue(const FGuid& FieldId, FString& OutJson) const;
+	bool ConstGetRemoteControlValue(const FGuid& FieldId, FString& Json) const;
 
 	UFUNCTION(BlueprintPure, Category="Render Grid Job")
-	bool GetRemoteControlValue(const FGuid& FieldId, FString& OutJson);
+	bool GetRemoteControlValue(const FGuid& FieldId, FString& Json);
 
 	UFUNCTION(BlueprintCallable, Category="Render Grid Job")
 	bool SetRemoteControlValue(const FGuid& FieldId, const FString& Json);
 
 	UFUNCTION(BlueprintPure, Category="Render Grid Job")
-	bool GetRemoteControlFieldIdFromLabel(const FString& Label, FGuid& OutFieldId);
+	bool GetRemoteControlFieldIdFromLabel(const FString& Label, FGuid& FieldId);
 
 	UFUNCTION(BlueprintPure, Category="Render Grid Job")
-	bool GetRemoteControlLabelFromFieldId(const FGuid& FieldId, FString& OutLabel);
+	bool GetRemoteControlLabelFromFieldId(const FGuid& FieldId, FString& Label);
 
 	UFUNCTION(BlueprintCallable, Category="Render Grid Job")
 	TMap<FGuid, FString> GetRemoteControlValues();
diff --git a/Engine/Plugins/Experimental/RenderGrid/Source/RenderGrid/Public/RenderGrid/RenderGridQueue.h b/Engine/Plugins/Experimental/RenderGrid/Source/RenderGrid/Public/RenderGrid/RenderGridQueue.h
--- a/Engine/Plugins/Experimental/RenderGrid/Source/RenderGrid/Public/RenderGrid/RenderGridQueue.h
+++ b/Engine/Plugins/Experimental/RenderGrid/Source/RenderGrid/Public/RenderGrid/RenderGridQueue.h
@@ -3,7 +3,7 @@
 #pragma once
 
 #include "RenderGridManager.h"
-#include "RenderGridUtils.h"
+#include "Utils/RenderGridUtils.h"
 #include "Tickable.h"
 #include "Containers/Queue.h"
 #include "UObject/ObjectPtr.h"
diff --git a/Engine/Plugins/Experimental/RenderGrid/Source/RenderGrid/Public/RenderGridUtils.h b/Engine/Plugins/Experimental/RenderGrid/Source/RenderGrid/Public/RenderGridUtils.h
deleted file mode 100644
diff --git a/Engine/Plugins/Experimental/RenderGrid/Source/RenderGrid/Public/Utils/RenderGridGenericExecutionQueue.h b/Engine/Plugins/Experimental/RenderGrid/Source/RenderGrid/Public/Utils/RenderGridGenericExecutionQueue.h
new file mode 100644
--- /dev/null
+++ b/Engine/Plugins/Experimental/RenderGrid/Source/RenderGrid/Public/Utils/RenderGridGenericExecutionQueue.h
@@ -0,0 +1,175 @@
+﻿// Copyright Epic Games, Inc. All Rights Reserved.
+
+#pragma once
+
+#include "CoreMinimal.h"
+#include "Async/Future.h"
+#include "Containers/List.h"
+#include "Tickable.h"
+
+
+namespace UE::RenderGrid::Private
+{
+	/**
+	 * Struct containing the delay data (such as the frames or the seconds of delay this delay requires before it can move on to the next step).
+	 */
+	struct RENDERGRID_API FRenderGridGenericExecutionQueueDelay
+	{
+	public:
+		static FRenderGridGenericExecutionQueueDelay Frames(const int64 Frames) { return FRenderGridGenericExecutionQueueDelay(Frames); }
+		static FRenderGridGenericExecutionQueueDelay Seconds(const double Seconds) { return FRenderGridGenericExecutionQueueDelay(Seconds); }
+		static FRenderGridGenericExecutionQueueDelay FramesOrSeconds(const int64 Frames, const double Seconds) { return FRenderGridGenericExecutionQueueDelay(Frames, Seconds); }
+
+	public:
+		FRenderGridGenericExecutionQueueDelay(TYPE_OF_NULLPTR)
+		{}
+
+	private:
+		FRenderGridGenericExecutionQueueDelay()
+		{}
+
+		explicit FRenderGridGenericExecutionQueueDelay(const int64 Frames)
+			: MinimumFrames(Frames)
+		{}
+
+		explicit FRenderGridGenericExecutionQueueDelay(const double Seconds)
+			: MinimumSeconds(Seconds)
+		{}
+
+		explicit FRenderGridGenericExecutionQueueDelay(const int64 Frames, const double Seconds)
+			: MinimumFrames(Frames)
+			, MinimumSeconds(Seconds)
+		{}
+
+	public:
+		const int64 MinimumFrames = 0;
+		const double MinimumSeconds = 0.0f;
+	};
+
+
+	/** A delegate for a queued action. */
+	DECLARE_DELEGATE(FRenderGridGenericExecutionQueueAction);
+
+	/** A delegate for a queued action, that optionally requires a delay after its execution. */
+	DECLARE_DELEGATE_RetVal(FRenderGridGenericExecutionQueueDelay, FRenderGridGenericExecutionQueueActionReturningDelay);
+
+	/** A delegate for a queued action, that will delay execution until the returned TFuture finishes. */
+	DECLARE_DELEGATE_RetVal(TSharedFuture<void>, FRenderGridGenericExecutionQueueActionReturningDelayFuture);
+
+	/** A delegate for a queued action, that will delay execution until the returned TFuture finishes, which can optionally return yet another delay if required. */
+	DECLARE_DELEGATE_RetVal(TSharedFuture<FRenderGridGenericExecutionQueueDelay>, FRenderGridGenericExecutionQueueActionReturningDelayFutureReturningDelay);
+
+	/**
+	 * Struct containing the data of a queued action.
+	 */
+	struct RENDERGRID_API FRenderGridGenericExecutionQueueEntry
+	{
+	public:
+		FRenderGridGenericExecutionQueueEntry(const FRenderGridGenericExecutionQueueAction& Action)
+			: ActionRegular(Action)
+		{}
+
+		FRenderGridGenericExecutionQueueEntry(const FRenderGridGenericExecutionQueueActionReturningDelay& Action)
+			: ActionReturningDelay(Action)
+		{}
+
+		FRenderGridGenericExecutionQueueEntry(const FRenderGridGenericExecutionQueueActionReturningDelayFuture& Action)
+			: ActionReturningDelayFuture(Action)
+		{}
+
+		FRenderGridGenericExecutionQueueEntry(const FRenderGridGenericExecutionQueueActionReturningDelayFutureReturningDelay& Action)
+			: ActionReturningDelayFutureReturningDelay(Action)
+		{}
+
+	public:
+		const FRenderGridGenericExecutionQueueAction ActionRegular;
+		const FRenderGridGenericExecutionQueueActionReturningDelay ActionReturningDelay;
+		const FRenderGridGenericExecutionQueueActionReturningDelayFuture ActionReturningDelayFuture;
+		const FRenderGridGenericExecutionQueueActionReturningDelayFutureReturningDelay ActionReturningDelayFutureReturningDelay;
+	};
+
+
+	/**
+	 * This class provides generic queue support, with built-in support for delays between actions.
+	 */
+	class RENDERGRID_API FRenderGridGenericExecutionQueue : public TSharedFromThis<FRenderGridGenericExecutionQueue>, public FTickableGameObject
+	{
+	public:
+		//~ Begin FTickableGameObject interface
+		virtual void Tick(float DeltaTime) override;
+		virtual ETickableTickType GetTickableTickType() const override { return ETickableTickType::Always; }
+		virtual bool IsTickableWhenPaused() const override { return true; }
+		virtual bool IsTickableInEditor() const override { return true; }
+		virtual bool IsTickable() const override { return true; }
+		virtual bool IsAllowedToTick() const override { return true; }
+		virtual TStatId GetStatId() const override
+		{
+			RETURN_QUICK_DECLARE_CYCLE_STAT(FRenderGridGenericExecutionQueue, STATGROUP_Tickables);
+		}
+		//~ End FTickableGameObject interface
+
+	public:
+		/** Queues the given action. */
+		void Add(const FRenderGridGenericExecutionQueueEntry& Entry) { QueuedEntries.AddTail(Entry); }
+
+		/** Queues the given delay. */
+		void Delay(const FRenderGridGenericExecutionQueueDelay& Delay) { Add(FRenderGridGenericExecutionQueueActionReturningDelay::CreateLambda([Delay]() -> FRenderGridGenericExecutionQueueDelay { return Delay; })); }
+
+		/** Queues the given delay, which will wait for the given number of frames. */
+		void DelayFrames(const int64 Frames) { Delay(FRenderGridGenericExecutionQueueDelay::Frames(Frames)); }
+
+		/** Queues the given delay, which will wait for the given number of seconds. */
+		void DelaySeconds(const double Seconds) { Delay(FRenderGridGenericExecutionQueueDelay::Seconds(Seconds)); }
+
+		/** Queues the given delay, which will wait for the given number of frames or seconds, whatever takes the longest. */
+		void DelayFramesOrSeconds(const int64 Frames, const double Seconds) { Delay(FRenderGridGenericExecutionQueueDelay::FramesOrSeconds(Frames, Seconds)); }
+
+		/** Starts the execution of this queue. */
+		void Start();
+
+		/** Stops (pauses) the execution of this queue, this can be resumed by calling the Start function again. Currently queued up delays will continue to run/expire. */
+		void Stop();
+
+		/** Returns true if Start() has been called and Stop() hasn't been called yet. */
+		bool IsRunning() { return bRunning; }
+
+	protected:
+		/** Executes the next delay (if there are any), otherwise it executes the next entry (action). */
+		void ExecuteNext();
+
+		/** Executes the next delay, returns true if it found any, returns false if there were no queued up delays. */
+		bool ExecuteNextDelay();
+
+		/** Executes the next entry (action), returns true if it found and executed an entry, returns false if there were no queued up entries. */
+		bool ExecuteNextEntry();
+
+		/** Adds the delay to the queued delays. */
+		void QueueDelay(const FRenderGridGenericExecutionQueueDelay& Delay) { QueuedDelays.AddTail(Delay); }
+
+	protected:
+		/** The queued up entries (actions). */
+		TDoubleLinkedList<FRenderGridGenericExecutionQueueEntry> QueuedEntries;
+
+		/** The queued up delays. */
+		TDoubleLinkedList<FRenderGridGenericExecutionQueueDelay> QueuedDelays;
+
+	protected:
+		/** Whether it has started (and hasn't been stopped/paused yet). This means that if this is true, Start() has been called, and Stop() hasn't been called yet since then. */
+		bool bRunning = false;
+
+		/** Whether it's currently executing a delay or an entry (action). False means there were no delays and entries queued up anymore. */
+		bool bExecuting = false;
+
+		/** The number of frames the current delay has to wait for. */
+		int64 DelayRemainingFrames = 0;
+
+		/** The number of seconds the current delay has to wait for. */
+		double DelayRemainingSeconds = 0.0f;
+
+		/** The TFuture it's waiting for (if any). */
+		TSharedFuture<void> DelayRemainingFuture;
+
+		/** The TFuture it's waiting for (if any), that can return a delay. */
+		TSharedFuture<FRenderGridGenericExecutionQueueDelay> DelayRemainingFutureReturningDelay;
+	};
+}
diff --git a/Engine/Plugins/Experimental/RenderGrid/Source/RenderGrid/Public/Utils/RenderGridRemoteControlUtils.h b/Engine/Plugins/Experimental/RenderGrid/Source/RenderGrid/Public/Utils/RenderGridRemoteControlUtils.h
new file mode 100644
--- /dev/null
+++ b/Engine/Plugins/Experimental/RenderGrid/Source/RenderGrid/Public/Utils/RenderGridRemoteControlUtils.h
@@ -0,0 +1,120 @@
+﻿// Copyright Epic Games, Inc. All Rights Reserved.
+
+#pragma once
+
+
+#include "CoreMinimal.h"
+#include "InstancedStruct.h"
+#include "Kismet/BlueprintFunctionLibrary.h"
+#include "RenderGridRemoteControlUtils.generated.h"
+
+
+UCLASS()
+class RENDERGRID_API URenderGridRemoteControlUtils : public UBlueprintFunctionLibrary
+{
+	GENERATED_UCLASS_BODY()
+
+public:
+	UFUNCTION(BlueprintPure, Category="Render Grid Utils", Meta=(Keywords="remote control"))
+	static void ParseJsonAsByte(const FString& Json, const uint8 DefaultValue, bool& bSuccess, uint8& Value);
+
+	UFUNCTION(BlueprintPure, Category="Render Grid Utils", Meta=(Keywords="remote control"))
+	static void ParseJsonAsInt32(const FString& Json, const int32 DefaultValue, bool& bSuccess, int32& Value);
+
+	UFUNCTION(BlueprintPure, Category="Render Grid Utils", Meta=(Keywords="remote control"))
+	static void ParseJsonAsInt64(const FString& Json, const int64 DefaultValue, bool& bSuccess, int64& Value);
+
+	UFUNCTION(BlueprintPure, Category="Render Grid Utils", Meta=(Keywords="remote control"))
+	static void ParseJsonAsFloat(const FString& Json, const float DefaultValue, bool& bSuccess, float& Value);
+
+	UFUNCTION(BlueprintPure, Category="Render Grid Utils", Meta=(Keywords="remote control"))
+	static void ParseJsonAsDouble(const FString& Json, const double DefaultValue, bool& bSuccess, double& Value);
+
+	UFUNCTION(BlueprintPure, Category="Render Grid Utils", Meta=(Keywords="remote control"))
+	static void ParseJsonAsBoolean(const FString& Json, bool DefaultValue, bool& bSuccess, bool& Value);
+
+	UFUNCTION(BlueprintPure, Category="Render Grid Utils", Meta=(Keywords="remote control"))
+	static void ParseJsonAsString(const FString& Json, const FString& DefaultValue, bool& bSuccess, FString& Value);
+
+	UFUNCTION(BlueprintPure, Category="Render Grid Utils", Meta=(Keywords="remote control"))
+	static void ParseJsonAsName(const FString& Json, const FName& DefaultValue, bool& bSuccess, FName& Value);
+
+	UFUNCTION(BlueprintPure, Category="Render Grid Utils", Meta=(Keywords="remote control"))
+	static void ParseJsonAsText(const FString& Json, const FText& DefaultValue, bool& bSuccess, FText& Value);
+
+	UFUNCTION(BlueprintPure, Category="Render Grid Utils", Meta=(Keywords="remote control"))
+	static void ParseJsonAsObjectReference(const FString& Json, UObject* DefaultValue, bool& bSuccess, UObject*& Value);
+
+	UFUNCTION(BlueprintPure, Category="Render Grid Utils", Meta=(Keywords="remote control"))
+	static void ParseJsonAsClassReference(const FString& Json, UClass* DefaultValue, bool& bSuccess, UClass*& Value);
+
+	UFUNCTION(BlueprintPure, Category="Render Grid Utils", Meta=(Keywords="remote control"))
+	static void ParseJsonAsStruct(const FString& Json, const FInstancedStruct& DefaultValue, bool& bSuccess, FInstancedStruct& Value);
+
+	UFUNCTION(BlueprintPure, Category="Render Grid Utils", Meta=(Keywords="remote control"))
+	static void ParseJsonAsVector(const FString& Json, const FVector& DefaultValue, bool& bSuccess, FVector& Value);
+
+	UFUNCTION(BlueprintPure, Category="Render Grid Utils", Meta=(Keywords="remote control"))
+	static void ParseJsonAsRotator(const FString& Json, const FRotator& DefaultValue, bool& bSuccess, FRotator& Value);
+
+	UFUNCTION(BlueprintPure, Category="Render Grid Utils", Meta=(Keywords="remote control"))
+	static void ParseJsonAsTransform(const FString& Json, const FTransform& DefaultValue, bool& bSuccess, FTransform& Value);
+
+	UFUNCTION(BlueprintPure, Category="Render Grid Utils", Meta=(Keywords="remote control"))
+	static void ParseJsonAsColor(const FString& Json, const FColor& DefaultValue, bool& bSuccess, FColor& Value);
+
+	UFUNCTION(BlueprintPure, Category="Render Grid Utils", Meta=(Keywords="remote control"))
+	static void ParseJsonAsLinearColor(const FString& Json, const FLinearColor& DefaultValue, bool& bSuccess, FLinearColor& Value);
+
+
+	UFUNCTION(BlueprintPure, Category="Render Grid Utils", Meta=(Keywords="remote control"))
+	static void ByteToJson(uint8 Value, FString& Json);
+
+	UFUNCTION(BlueprintPure, Category="Render Grid Utils", Meta=(Keywords="remote control"))
+	static void Int32ToJson(int32 Value, FString& Json);
+
+	UFUNCTION(BlueprintPure, Category="Render Grid Utils", Meta=(Keywords="remote control"))
+	static void Int64ToJson(int64 Value, FString& Json);
+
+	UFUNCTION(BlueprintPure, Category="Render Grid Utils", Meta=(Keywords="remote control"))
+	static void FloatToJson(float Value, FString& Json);
+
+	UFUNCTION(BlueprintPure, Category="Render Grid Utils", Meta=(Keywords="remote control"))
+	static void DoubleToJson(double Value, FString& Json);
+
+	UFUNCTION(BlueprintPure, Category="Render Grid Utils", Meta=(Keywords="remote control"))
+	static void BooleanToJson(bool Value, FString& Json);
+
+	UFUNCTION(BlueprintPure, Category="Render Grid Utils", Meta=(Keywords="remote control"))
+	static void StringToJson(const FString& Value, FString& Json);
+
+	UFUNCTION(BlueprintPure, Category="Render Grid Utils", Meta=(Keywords="remote control"))
+	static void NameToJson(const FName& Value, FString& Json);
+
+	UFUNCTION(BlueprintPure, Category="Render Grid Utils", Meta=(Keywords="remote control"))
+	static void TextToJson(const FText& Value, FString& Json);
+
+	UFUNCTION(BlueprintPure, Category="Render Grid Utils", Meta=(Keywords="remote control"))
+	static void ObjectReferenceToJson(UObject* Value, FString& Json);
+
+	UFUNCTION(BlueprintPure, Category="Render Grid Utils", Meta=(Keywords="remote control"))
+	static void ClassReferenceToJson(UClass* Value, FString& Json);
+
+	UFUNCTION(BlueprintPure, Category="Render Grid Utils", Meta=(Keywords="remote control"))
+	static void StructToJson(const FInstancedStruct& Value, FString& Json);
+
+	UFUNCTION(BlueprintPure, Category="Render Grid Utils", Meta=(Keywords="remote control"))
+	static void VectorToJson(const FVector& Value, FString& Json);
+
+	UFUNCTION(BlueprintPure, Category="Render Grid Utils", Meta=(Keywords="remote control"))
+	static void RotatorToJson(const FRotator& Value, FString& Json);
+
+	UFUNCTION(BlueprintPure, Category="Render Grid Utils", Meta=(Keywords="remote control"))
+	static void TransformToJson(const FTransform& Value, FString& Json);
+
+	UFUNCTION(BlueprintPure, Category="Render Grid Utils", Meta=(Keywords="remote control"))
+	static void ColorToJson(const FColor& Value, FString& Json);
+
+	UFUNCTION(BlueprintPure, Category="Render Grid Utils", Meta=(Keywords="remote control"))
+	static void LinearColorToJson(const FLinearColor& Value, FString& Json);
+};
diff --git a/Engine/Plugins/Experimental/RenderGrid/Source/RenderGrid/Public/Utils/RenderGridUtils.h b/Engine/Plugins/Experimental/RenderGrid/Source/RenderGrid/Public/Utils/RenderGridUtils.h
new file mode 100644
--- /dev/null
+++ b/Engine/Plugins/Experimental/RenderGrid/Source/RenderGrid/Public/Utils/RenderGridUtils.h
@@ -0,0 +1,204 @@
+// Copyright Epic Games, Inc. All Rights Reserved.
+
+#pragma once
+
+#include "CoreMinimal.h"
+#include "IImageWrapper.h"
+#include "RenderGridUtils.generated.h"
+
+
+class URenderGrid;
+class UTexture2D;
+class FJsonValue;
+
+
+/**
+ * This struct keeps track of the values of the GEngine framerate settings before new values were applied, so we can rollback to the previous state.
+ */
+USTRUCT()
+struct RENDERGRID_API FRenderGridPreviousEngineFpsSettings
+{
+	GENERATED_BODY()
+
+public:
+	/** Whether the values have been set or not. */
+	UPROPERTY()
+	bool bHasBeenSet = false;
+
+	/** The previous value of GEngine->bUseFixedFrameRate. */
+	UPROPERTY()
+	bool bUseFixedFrameRate = false;
+
+	/** The previous value of GEngine->bForceDisableFrameRateSmoothing. */
+	UPROPERTY()
+	bool bForceDisableFrameRateSmoothing = false;
+
+	/** The previous value of GEngine->GetMaxFPS(). */
+	UPROPERTY()
+	float MaxFps = 0;
+
+	/** The previous value of console variable "r.VSync". */
+	UPROPERTY()
+	bool bVSync = false;
+
+	/** The previous value of console variable "r.VSyncEditor". */
+	UPROPERTY()
+	bool bVSyncEditor = false;
+
+	/** The previous value of UEditorPerformanceSettings->bThrottleCPUWhenNotForeground. */
+	UPROPERTY()
+	bool bThrottleCPUWhenNotForeground = false;
+};
+
+
+namespace UE::RenderGrid::Private
+{
+	/**
+	 * A class containing static utility functions for the RenderGrid module.
+	 */
+	class RENDERGRID_API FRenderGridUtils
+	{
+	public:
+		/**
+		 * Returns true if the given file is likely a valid image.
+		 */
+		static bool IsImage(const FString& ImagePath);
+
+		/**
+		 * Loads an image from the disk, tries to automatically figure out the correct image format.
+		 *
+		 * Returns NULL if it fails.
+		 *
+		 * Will re-use the given Texture2D if possible, bOutReusedGivenTexture2D will be true if it was possible.
+		 */
+		static UTexture2D* GetImage(const FString& ImagePath, UTexture2D* Texture2D, bool& bOutReusedGivenTexture2D);
+
+
+		/**
+		 * Converts bytes into a texture.
+		 *
+		 * Returns NULL if it fails.
+		 */
+		static UTexture2D* BytesToImage(const TArray<uint8>& ByteArray, const EImageFormat ImageFormat);
+
+		/**
+		 * Converts bytes into an texture.
+		 *
+		 * Returns NULL if it fails.
+		 *
+		 * Will re-use the given Texture2D if possible, bOutReusedGivenTexture2D will be true if it was possible.
+		 */
+		static UTexture2D* BytesToExistingImage(const TArray<uint8>& ByteArray, const EImageFormat ImageFormat, UTexture2D* Texture2D, bool& bOutReusedGivenTexture2D);
+
+		/**
+		 * Converts texture data into a texture.
+		 *
+		 * Returns NULL if it fails.
+		 */
+		static UTexture2D* DataToTexture2D(int32 Width, int32 Height, const void* Src, SIZE_T Count);
+
+		/**
+		 * Converts texture data into an texture.
+		 *
+		 * Returns NULL if it fails.
+		 *
+		 * Will re-use the given Texture2D if possible, bOutReusedGivenTexture2D will be true if it was possible.
+		 */
+		static UTexture2D* DataToExistingTexture2D(int32 Width, int32 Height, const void* Src, SIZE_T Count, UTexture2D* Texture2D, bool& bOutReusedGivenTexture2D);
+
+
+		/**
+		 * Returns the paths of the files that exist in the given directory path.
+		 */
+		static TArray<FString> GetFiles(const FString& Directory, const bool bRecursive = false);
+
+		/**
+		 * Returns the data of the file, returns an empty byte array if the file doesn't exist.
+		 *
+		 * Note: Can only open files of 2GB and smaller, will return an empty byte array if it is bigger than 2GB.
+		 */
+		static TArray<uint8> GetFileData(const FString& File);
+
+
+		/**
+		 * Deletes all files and directories in the given directory, including the given directory.
+		 */
+		static void DeleteDirectory(const FString& Directory);
+
+		/**
+		 * Deletes all files and directories in the given directory, won't delete the given directory.
+		 */
+		static void EmptyDirectory(const FString& Directory);
+
+
+		/**
+		 * Returns a normalized directory path.
+		 */
+		static FString NormalizeOutputDirectory(const FString& OutputDirectory);
+
+
+		/**
+		 * Returns a job ID with all invalid characters removed from it. Returns an empty string if it contains no valid characters.
+		 */
+		static FString PurgeJobIdOrReturnEmptyString(const FString& JobId);
+
+		/**
+		 * Returns a job ID with all invalid characters removed from it. Returns string "0" if it contains no valid characters.
+		 */
+		static FString PurgeJobId(const FString& JobId);
+
+		/**
+		 * Returns a job ID with all invalid characters removed from it. Returns a newly generated unique ID if it contains no valid characters.
+		 * The generated string is only unique relative to the job IDs that currently exist in the given RenderGrid.
+		 */
+		static FString PurgeJobIdOrGenerateUniqueId(URenderGrid* Grid, const FString& JobId);
+
+		/**
+		 * Returns a job name with all invalid characters removed from it. Returns an empty string if it contains no valid characters.
+		 */
+		static FString PurgeJobName(const FString& JobName);
+
+		/**
+		 * Returns a job output directory that's converted to a full path (if it's currently relative) and that also has the project directory path replaced by "{project_dir}" (if the given path contains that string).
+		 */
+		static FString NormalizeJobOutputDirectory(const FString& NewOutputDirectory);
+
+		/**
+		 * Returns a job output directory that's a usable valid filesystem path. Replaces string "{project_dir}" by the project directory.
+		 */
+		static FString DenormalizeJobOutputDirectory(const FString& NormalizedOutputDirectory);
+
+
+		/**
+		 * Converts the given FJsonValue to a string.
+		 */
+		static FString ToJsonString(const TSharedPtr<FJsonValue>& Value, const bool PrettyPrint = false);
+
+		/**
+		 * Parses the given JSON string and returns the FJsonValue. Will return an invalid pointer if the given JSON couldn't be parsed.
+		 */
+		static TSharedPtr<FJsonValue> ParseJson(const FString& Json);
+
+
+		/**
+		 * Converts the low-level FProperty binary data to a clean JSON string.
+		 */
+		static FString GetRemoteControlValueJsonFromBytes(const TArray<uint8>& ValueBytes, const bool PrettyPrint = false);
+
+		/**
+		 * Converts a clean JSON string to low-level FProperty binary data.
+		 */
+		static TArray<uint8> GetRemoteControlValueBytesFromJson(const TArray<uint8>& OldValueBytes, const FString& NewValueJson);
+
+
+		/**
+		 * Disables the current FPS limiting, returns the previous settings which allow you to revert back to the previous state.
+		 */
+		static FRenderGridPreviousEngineFpsSettings DisableFpsLimit();
+
+		/**
+		 * Reverts back to the previous (given) state.
+		 */
+		static void RestoreFpsLimit(const FRenderGridPreviousEngineFpsSettings& Settings);
+	};
+}
diff --git a/Engine/Plugins/Experimental/RenderGrid/Source/RenderGrid/RenderGrid.Build.cs b/Engine/Plugins/Experimental/RenderGrid/Source/RenderGrid/RenderGrid.Build.cs
--- a/Engine/Plugins/Experimental/RenderGrid/Source/RenderGrid/RenderGrid.Build.cs
+++ b/Engine/Plugins/Experimental/RenderGrid/Source/RenderGrid/RenderGrid.Build.cs
@@ -17,6 +17,7 @@ public class RenderGrid : ModuleRules
 				"CoreUObject",
 				"Engine",
 				"RemoteControl",
+				"StructUtils",
 			}
 		);
 
diff --git a/Engine/Plugins/Experimental/RenderGrid/Source/RenderGridEditor/Private/UI/SRenderGridJobList.cpp b/Engine/Plugins/Experimental/RenderGrid/Source/RenderGridEditor/Private/UI/SRenderGridJobList.cpp
--- a/Engine/Plugins/Experimental/RenderGrid/Source/RenderGridEditor/Private/UI/SRenderGridJobList.cpp
+++ b/Engine/Plugins/Experimental/RenderGrid/Source/RenderGridEditor/Private/UI/SRenderGridJobList.cpp
@@ -9,7 +9,7 @@
 #include "RenderGrid/RenderGrid.h"
 #include "RenderGrid/RenderGridQueue.h"
 #include "IRenderGridEditor.h"
-#include "RenderGridUtils.h"
+#include "Utils/RenderGridUtils.h"
 
 #include "Framework/MultiBox/MultiBoxBuilder.h"
 #include "Misc/MessageDialog.h"
