RenderGrid plugin:
- Removed the double to/from json blueprints, instead make the float blueprints always accept/return a double value.

diff --git a/Engine/Plugins/Experimental/RenderGrid/Source/RenderGrid/Private/Utils/RenderGridRemoteControlUtils.cpp b/Engine/Plugins/Experimental/RenderGrid/Source/RenderGrid/Private/Utils/RenderGridRemoteControlUtils.cpp
--- a/Engine/Plugins/Experimental/RenderGrid/Source/RenderGrid/Private/Utils/RenderGridRemoteControlUtils.cpp
+++ b/Engine/Plugins/Experimental/RenderGrid/Source/RenderGrid/Private/Utils/RenderGridRemoteControlUtils.cpp
@@ -107,12 +107,7 @@ void URenderGridRemoteControlUtils::ParseJsonAsInt64(const FString& Json, const
 	ParseJsonAsNumericType<int64>(Json, DefaultValue, Value, bSuccess);
 }
 
-void URenderGridRemoteControlUtils::ParseJsonAsFloat(const FString& Json, const float DefaultValue, bool& bSuccess, float& Value)
-{
-	ParseJsonAsNumericType<float>(Json, DefaultValue, Value, bSuccess);
-}
-
-void URenderGridRemoteControlUtils::ParseJsonAsDouble(const FString& Json, const double DefaultValue, bool& bSuccess, double& Value)
+void URenderGridRemoteControlUtils::ParseJsonAsFloat(const FString& Json, const double DefaultValue, bool& bSuccess, double& Value)
 {
 	ParseJsonAsNumericType<double>(Json, DefaultValue, Value, bSuccess);
 }
@@ -300,12 +295,7 @@ void URenderGridRemoteControlUtils::Int64ToJson(const int64 Value, FString& Json
 	Json = UE::RenderGrid::Private::FRenderGridUtils::ToJsonString(MakeShareable(new FJsonValueNumber(Value)));
 }
 
-void URenderGridRemoteControlUtils::FloatToJson(const float Value, FString& Json)
-{
-	Json = UE::RenderGrid::Private::FRenderGridUtils::ToJsonString(MakeShareable(new FJsonValueNumber(Value)));
-}
-
-void URenderGridRemoteControlUtils::DoubleToJson(const double Value, FString& Json)
+void URenderGridRemoteControlUtils::FloatToJson(const double Value, FString& Json)
 {
 	Json = UE::RenderGrid::Private::FRenderGridUtils::ToJsonString(MakeShareable(new FJsonValueNumber(Value)));
 }
diff --git a/Engine/Plugins/Experimental/RenderGrid/Source/RenderGrid/Public/Utils/RenderGridRemoteControlUtils.h b/Engine/Plugins/Experimental/RenderGrid/Source/RenderGrid/Public/Utils/RenderGridRemoteControlUtils.h
--- a/Engine/Plugins/Experimental/RenderGrid/Source/RenderGrid/Public/Utils/RenderGridRemoteControlUtils.h
+++ b/Engine/Plugins/Experimental/RenderGrid/Source/RenderGrid/Public/Utils/RenderGridRemoteControlUtils.h
@@ -24,11 +24,8 @@ public:
 	UFUNCTION(BlueprintPure, Category="Render Grid Utils", Meta=(Keywords="remote control"))
 	static void ParseJsonAsInt64(const FString& Json, const int64 DefaultValue, bool& bSuccess, int64& Value);
 
-	UFUNCTION(BlueprintPure, Category="Render Grid Utils", Meta=(Keywords="remote control"))
-	static void ParseJsonAsFloat(const FString& Json, const float DefaultValue, bool& bSuccess, float& Value);
-
-	UFUNCTION(BlueprintPure, Category="Render Grid Utils", Meta=(Keywords="remote control"))
-	static void ParseJsonAsDouble(const FString& Json, const double DefaultValue, bool& bSuccess, double& Value);
+	UFUNCTION(BlueprintPure, Category="Render Grid Utils", Meta=(Keywords="remote control double"))
+	static void ParseJsonAsFloat(const FString& Json, const double DefaultValue, bool& bSuccess, double& Value);
 
 	UFUNCTION(BlueprintPure, Category="Render Grid Utils", Meta=(Keywords="remote control"))
 	static void ParseJsonAsBoolean(const FString& Json, bool DefaultValue, bool& bSuccess, bool& Value);
@@ -76,11 +73,8 @@ public:
 	UFUNCTION(BlueprintPure, Category="Render Grid Utils", Meta=(Keywords="remote control"))
 	static void Int64ToJson(int64 Value, FString& Json);
 
-	UFUNCTION(BlueprintPure, Category="Render Grid Utils", Meta=(Keywords="remote control"))
-	static void FloatToJson(float Value, FString& Json);
-
-	UFUNCTION(BlueprintPure, Category="Render Grid Utils", Meta=(Keywords="remote control"))
-	static void DoubleToJson(double Value, FString& Json);
+	UFUNCTION(BlueprintPure, Category="Render Grid Utils", Meta=(Keywords="remote control double"))
+	static void FloatToJson(double Value, FString& Json);
 
 	UFUNCTION(BlueprintPure, Category="Render Grid Utils", Meta=(Keywords="remote control"))
 	static void BooleanToJson(bool Value, FString& Json);
