RenderGrid plugin:
- Updated the project from using an AssetTypeAction to an AssetDefinition.
#p4v-preflight-copy 23974134
#p4v-preflight-copy 23974267

diff --git a/Engine/Plugins/Experimental/RenderGrid/Source/RenderGrid/Private/Utils/RenderGridRemoteControlUtils.cpp b/Engine/Plugins/Experimental/RenderGrid/Source/RenderGrid/Private/Utils/RenderGridRemoteControlUtils.cpp
--- a/Engine/Plugins/Experimental/RenderGrid/Source/RenderGrid/Private/Utils/RenderGridRemoteControlUtils.cpp
+++ b/Engine/Plugins/Experimental/RenderGrid/Source/RenderGrid/Private/Utils/RenderGridRemoteControlUtils.cpp
@@ -9,8 +9,8 @@
 #include "Backends/JsonStructSerializerBackend.h"
 #include "Utils/RenderGridUtils.h"
 #include "Dom/JsonValue.h"
-#include "Engine/Engine.h"
 #include "Serialization/MemoryReader.h"
+#include "UObject/StructOnScope.h"
 
 
 URenderGridRemoteControlUtils::URenderGridRemoteControlUtils(const FObjectInitializer& ObjectInitializer) : Super(ObjectInitializer) {}
diff --git a/Engine/Plugins/Experimental/RenderGrid/Source/RenderGridEditor/Private/AssetDefinition/RenderGridBlueprintAssetDefinition.cpp b/Engine/Plugins/Experimental/RenderGrid/Source/RenderGridEditor/Private/AssetDefinition/RenderGridBlueprintAssetDefinition.cpp
new file mode 100644
--- /dev/null
+++ b/Engine/Plugins/Experimental/RenderGrid/Source/RenderGridEditor/Private/AssetDefinition/RenderGridBlueprintAssetDefinition.cpp
@@ -0,0 +1,67 @@
+// Copyright Epic Games, Inc. All Rights Reserved.
+
+#include "RenderGridBlueprintAssetDefinition.h"
+#include "Blueprints/RenderGridBlueprint.h"
+#include "RenderGridEditorModule.h"
+#include "Misc/MessageDialog.h"
+
+#define LOCTEXT_NAMESPACE "RenderGridBlueprintAssetDefinition"
+
+
+FText URenderGridBlueprintAssetDefinition::GetAssetDisplayName() const
+{
+	return LOCTEXT("AssetTypeActions_RenderGrid", "Render Grid");
+}
+
+FLinearColor URenderGridBlueprintAssetDefinition::GetAssetColor() const
+{
+	return FLinearColor(FColor(200, 80, 80));
+}
+
+TSoftClassPtr<UObject> URenderGridBlueprintAssetDefinition::GetAssetClass() const
+{
+	return URenderGridBlueprint::StaticClass();
+}
+
+TConstArrayView<FAssetCategoryPath> URenderGridBlueprintAssetDefinition::GetAssetCategories() const
+{
+	static const auto Categories = {EAssetCategoryPaths::Misc};
+	return Categories;
+}
+
+FAssetSupportResponse URenderGridBlueprintAssetDefinition::CanLocalize(const FAssetData& InAsset) const
+{
+	return FAssetSupportResponse::NotSupported();
+}
+
+EAssetCommandResult URenderGridBlueprintAssetDefinition::OpenAssets(const FAssetOpenArgs& OpenArgs) const
+{
+	TArray<FAssetData> OutAssetsThatFailedToLoad;
+	for (URenderGridBlueprint* RenderGridBlueprint : OpenArgs.LoadObjects<URenderGridBlueprint>({}, &OutAssetsThatFailedToLoad))
+	{
+		bool bLetOpen = true;
+		if (!RenderGridBlueprint->SkeletonGeneratedClass || !RenderGridBlueprint->GeneratedClass)
+		{
+			bLetOpen = EAppReturnType::Yes == FMessageDialog::Open(EAppMsgType::YesNo, LOCTEXT("FailedToLoadBlueprintWithContinue", "Blueprint could not be loaded because it derives from an invalid class.  Check to make sure the parent class for this blueprint hasn't been removed! Do you want to continue (it can crash the editor)?"));
+		}
+		if (bLetOpen)
+		{
+			UE::RenderGrid::IRenderGridEditorModule::Get().CreateRenderGridEditor(OpenArgs.GetToolkitMode(), OpenArgs.ToolkitHost, RenderGridBlueprint);
+		}
+	}
+
+	for (const FAssetData& UnableToLoadAsset : OutAssetsThatFailedToLoad)
+	{
+		FMessageDialog::Open(EAppMsgType::Ok,
+			FText::Format(
+				LOCTEXT("FailedToLoadBlueprint", "Blueprint '{0}' could not be loaded because it derives from an invalid class.  Check to make sure the parent class for this blueprint hasn't been removed!"),
+				FText::FromName(UnableToLoadAsset.PackagePath)
+			)
+		);
+	}
+
+	return EAssetCommandResult::Handled;
+}
+
+
+#undef LOCTEXT_NAMESPACE
diff --git a/Engine/Plugins/Experimental/RenderGrid/Source/RenderGridEditor/Private/AssetDefinition/RenderGridBlueprintAssetDefinition.h b/Engine/Plugins/Experimental/RenderGrid/Source/RenderGridEditor/Private/AssetDefinition/RenderGridBlueprintAssetDefinition.h
new file mode 100644
--- /dev/null
+++ b/Engine/Plugins/Experimental/RenderGrid/Source/RenderGridEditor/Private/AssetDefinition/RenderGridBlueprintAssetDefinition.h
@@ -0,0 +1,26 @@
+// Copyright Epic Games, Inc. All Rights Reserved.
+
+#pragma once
+
+#include "AssetDefinitionDefault.h"
+#include "RenderGridBlueprintAssetDefinition.generated.h"
+
+
+/**
+ * This class adds support for the RenderGrid (RenderGridBlueprint) class to the AssetDefinition module.
+ */
+UCLASS()
+class URenderGridBlueprintAssetDefinition : public UAssetDefinitionDefault
+{
+	GENERATED_BODY()
+
+public:
+	//~ Begin UAssetDefinition Interface
+	virtual FText GetAssetDisplayName() const override;
+	virtual FLinearColor GetAssetColor() const override;
+	virtual TSoftClassPtr<UObject> GetAssetClass() const override;
+	virtual TConstArrayView<FAssetCategoryPath> GetAssetCategories() const override;
+	virtual FAssetSupportResponse CanLocalize(const FAssetData& InAsset) const override;
+	virtual EAssetCommandResult OpenAssets(const FAssetOpenArgs& OpenArgs) const override;
+	//~ End UAssetDefinition interface
+};
diff --git a/Engine/Plugins/Experimental/RenderGrid/Source/RenderGridEditor/Private/AssetTypeActions/RenderGridBlueprintActions.cpp b/Engine/Plugins/Experimental/RenderGrid/Source/RenderGridEditor/Private/AssetTypeActions/RenderGridBlueprintActions.cpp
deleted file mode 100644
diff --git a/Engine/Plugins/Experimental/RenderGrid/Source/RenderGridEditor/Private/AssetTypeActions/RenderGridBlueprintActions.h b/Engine/Plugins/Experimental/RenderGrid/Source/RenderGridEditor/Private/AssetTypeActions/RenderGridBlueprintActions.h
deleted file mode 100644
diff --git a/Engine/Plugins/Experimental/RenderGrid/Source/RenderGridEditor/Private/RenderGridEditorModule.cpp b/Engine/Plugins/Experimental/RenderGrid/Source/RenderGridEditor/Private/RenderGridEditorModule.cpp
--- a/Engine/Plugins/Experimental/RenderGrid/Source/RenderGridEditor/Private/RenderGridEditorModule.cpp
+++ b/Engine/Plugins/Experimental/RenderGrid/Source/RenderGridEditor/Private/RenderGridEditorModule.cpp
@@ -2,17 +2,13 @@
 
 #include "RenderGridEditorModule.h"
 
-#include "AssetTypeActions/RenderGridBlueprintActions.h"
 #include "Blueprints/RenderGridBlueprint.h"
 #include "Commands/RenderGridEditorCommands.h"
-#include "Factories/RenderGridFactory.h"
 #include "Factories/RenderGridPropsSourceWidgetFactoryLocal.h"
 #include "Factories/RenderGridPropsSourceWidgetFactoryRemoteControl.h"
 #include "RenderGrid/RenderGrid.h"
 #include "Styles/RenderGridEditorStyle.h"
 #include "Toolkit/RenderGridEditor.h"
-
-#include "IAssetTools.h"
 #include "Kismet2/KismetEditorUtilities.h"
 
 #define LOCTEXT_NAMESPACE "RenderGridEditorModule"
@@ -29,32 +25,12 @@ void UE::RenderGrid::Private::FRenderGridEditorModule::StartupModule()
 
 	RegisterPropsSourceWidgetFactories();
 
-	// Register asset tools
-	auto RegisterAssetTypeAction = [this](const TSharedRef<IAssetTypeActions>& InAssetTypeAction)
-	{
-		IAssetTools& AssetTools = FModuleManager::LoadModuleChecked<FAssetToolsModule>("AssetTools").Get();
-		RegisteredAssetTypeActions.Add(InAssetTypeAction);
-		AssetTools.RegisterAssetTypeActions(InAssetTypeAction);
-	};
-
-	RegisterAssetTypeAction(MakeShared<FRenderGridBlueprintActions>());
-
 	// Register to fixup newly created BPs
 	FKismetEditorUtilities::RegisterOnBlueprintCreatedCallback(this, URenderGrid::StaticClass(), FKismetEditorUtilities::FOnBlueprintCreated::CreateRaw(this, &FRenderGridEditorModule::HandleNewBlueprintCreated));
 }
 
 void UE::RenderGrid::Private::FRenderGridEditorModule::ShutdownModule()
 {
-	FAssetToolsModule* AssetToolsModule = FModuleManager::GetModulePtr<FAssetToolsModule>("AssetTools");
-
-	if (AssetToolsModule)
-	{
-		for (TSharedRef<IAssetTypeActions> RegisteredAssetTypeAction : RegisteredAssetTypeActions)
-		{
-			AssetToolsModule->Get().UnregisterAssetTypeActions(RegisteredAssetTypeAction);
-		}
-	}
-
 	UnregisterPropsSourceWidgetFactories();
 
 	MenuExtensibilityManager.Reset();
diff --git a/Engine/Plugins/Experimental/RenderGrid/Source/RenderGridEditor/Private/RenderGridEditorModule.h b/Engine/Plugins/Experimental/RenderGrid/Source/RenderGridEditor/Private/RenderGridEditorModule.h
--- a/Engine/Plugins/Experimental/RenderGrid/Source/RenderGridEditor/Private/RenderGridEditorModule.h
+++ b/Engine/Plugins/Experimental/RenderGrid/Source/RenderGridEditor/Private/RenderGridEditorModule.h
@@ -6,16 +6,10 @@
 #include "IRenderGridEditorModule.h"
 
 
-class IAssetTypeActions;
 class IToolkitHost;
 class URenderGrid;
 class URenderGridBlueprint;
 
-namespace UE::RenderGrid::Private
-{
-	class FRenderGridEditor;
-}
-
 
 namespace UE::RenderGrid::Private
 {
@@ -53,8 +47,6 @@ namespace UE::RenderGrid::Private
 		TSharedPtr<FExtensibilityManager> MenuExtensibilityManager;
 		TSharedPtr<FExtensibilityManager> ToolBarExtensibilityManager;
 
-		TArray<TSharedRef<IAssetTypeActions>> RegisteredAssetTypeActions;
-
 		TMap<ERenderGridPropsSourceType, TSharedPtr<IRenderGridPropsSourceWidgetFactory>> PropsSourceWidgetFactories;
 	};
 }
diff --git a/Engine/Plugins/Experimental/RenderGrid/Source/RenderGridEditor/RenderGridEditor.Build.cs b/Engine/Plugins/Experimental/RenderGrid/Source/RenderGridEditor/RenderGridEditor.Build.cs
--- a/Engine/Plugins/Experimental/RenderGrid/Source/RenderGridEditor/RenderGridEditor.Build.cs
+++ b/Engine/Plugins/Experimental/RenderGrid/Source/RenderGridEditor/RenderGridEditor.Build.cs
@@ -24,7 +24,7 @@ public class RenderGridEditor : ModuleRules
 			new string[]
 			{
 				"AppFramework",
-				"AssetTools",
+				"AssetDefinition",
 				"DesktopPlatform",
 				"EditorStyle",
 				"GraphEditor",
