RenderGrid plugin:
- Added easy debugging blueprints/functions for the RenderGrid, RenderGridJob, and RenderGridQueue classes.
- Added several blueprints to the RenderGridQueue class, to be able to obtain data that wasn't possible before (like individual job completion percentages, as well as the currently rendering queue, etc.

diff --git a/Engine/Plugins/Experimental/RenderGrid/Source/RenderGrid/Private/RenderGrid/RenderGrid.cpp b/Engine/Plugins/Experimental/RenderGrid/Source/RenderGrid/Private/RenderGrid/RenderGrid.cpp
--- a/Engine/Plugins/Experimental/RenderGrid/Source/RenderGrid/Private/RenderGrid/RenderGrid.cpp
+++ b/Engine/Plugins/Experimental/RenderGrid/Source/RenderGrid/Private/RenderGrid/RenderGrid.cpp
@@ -2,12 +2,14 @@
 
 #include "RenderGrid/RenderGrid.h"
 
+#include "DiffUtils.h"
 #include "IRenderGridModule.h"
 #include "LevelSequence.h"
 #include "Misc/Base64.h"
 #include "MoviePipelineOutputSetting.h"
 #include "MoviePipelinePrimaryConfig.h"
 #include "MovieScene.h"
+#include "Dom/JsonObject.h"
 #include "RenderGrid/RenderGridManager.h"
 #include "RenderGrid/RenderGridQueue.h"
 #include "Utils/RenderGridUtils.h"
@@ -50,6 +52,75 @@ URenderGridJob::URenderGridJob()
 	SetFlags(RF_Public | RF_Transactional);
 }
 
+FString URenderGridJob::ToDebugString() const
+{
+	return UE::RenderGrid::Private::FRenderGridUtils::ToJsonString(ToDebugJson(), true);
+}
+
+TSharedPtr<FJsonValue> URenderGridJob::ToDebugJson() const
+{
+	TSharedPtr<FJsonObject> Root = MakeShareable(new FJsonObject());
+	{
+		Root->SetStringField(TEXT("Class"), GetParentNativeClass(GetClass())->GetName());
+		Root->SetStringField(TEXT("Object"), GetName());
+		Root->SetStringField(TEXT("Guid"), GetGuid().ToString());
+		Root->SetStringField(TEXT("JobId"), GetJobId());
+		Root->SetStringField(TEXT("JobName"), GetJobName());
+		Root->SetBoolField(TEXT("IsEnabled"), GetIsEnabled());
+		Root->SetStringField(TEXT("RenderPreset"), GetPathNameSafe(GetRenderPreset()));
+		Root->SetStringField(TEXT("OutputDirectory"), GetOutputDirectory());
+
+		{
+			const FIntPoint Value = GetOutputResolution();
+			TSharedPtr<FJsonObject> Object = MakeShareable(new FJsonObject());
+			Object->SetNumberField(TEXT("X"), Value.X);
+			Object->SetNumberField(TEXT("Y"), Value.Y);
+			Root->SetObjectField(TEXT("OutputResolution"), Object);
+		}
+
+		if (const TOptional<int32> Value = GetStartFrame(); Value.IsSet())
+		{
+			Root->SetNumberField(TEXT("StartFrame"), Value.Get(0));
+		}
+		if (const TOptional<int32> Value = GetEndFrame(); Value.IsSet())
+		{
+			Root->SetNumberField(TEXT("EndFrame"), Value.Get(0));
+		}
+		if (const TOptional<double> Value = GetStartTime(); Value.IsSet())
+		{
+			Root->SetNumberField(TEXT("StartTime"), Value.Get(0));
+		}
+		if (const TOptional<double> Value = GetEndTime(); Value.IsSet())
+		{
+			Root->SetNumberField(TEXT("EndTime"), Value.Get(0));
+		}
+		if (const TOptional<double> Value = GetDuration(); Value.IsSet())
+		{
+			Root->SetNumberField(TEXT("Duration"), Value.Get(0));
+		}
+
+		Root->SetNumberField(TEXT("WaitFramesBeforeRendering"), GetWaitFramesBeforeRendering());
+		Root->SetStringField(TEXT("LevelSequence"), GetPathNameSafe(GetLevelSequence()));
+
+		{
+			TArray<TSharedPtr<FJsonValue>> Array;
+			for (const TTuple<FGuid, FString>& Entry : GetRemoteControlValues())
+			{
+				TSharedPtr<FJsonObject> Object = MakeShareable(new FJsonObject());
+				Object->SetStringField(TEXT("Guid"), Entry.Key.ToString());
+				if (FString Label; GetRemoteControlLabelFromFieldId(Entry.Key, Label))
+				{
+					Object->SetStringField(TEXT("Label"), Label);
+				}
+				Object->SetField(TEXT("Value"), UE::RenderGrid::Private::FRenderGridUtils::ParseJson(Entry.Value));
+				Array.Add(MakeShareable(new FJsonValueObject(Object)));
+			}
+			Root->SetArrayField(TEXT("RemoteControlValues"), Array);
+		}
+	}
+	return MakeShareable(new FJsonValueObject(Root));
+}
+
 TOptional<int32> URenderGridJob::GetSequenceStartFrame() const
 {
 	if (!IsValid(LevelSequence))
@@ -268,7 +339,7 @@ TOptional<double> URenderGridJob::GetEndTime() const
 	return EndFrame.Get(0) / DisplayRate.AsDecimal();
 }
 
-TOptional<double> URenderGridJob::GetDurationInSeconds() const
+TOptional<double> URenderGridJob::GetDuration() const
 {
 	if (!IsValid(LevelSequence))
 	{
@@ -293,6 +364,41 @@ TOptional<double> URenderGridJob::GetDurationInSeconds() const
 	return (EndFrame.Get(0) - StartFrame.Get(0)) / DisplayRate.AsDecimal();
 }
 
+void URenderGridJob::GetStartFrame(bool& bSuccess, int32& StartFrame) const
+{
+	TOptional<int32> Result = GetStartFrame();
+	bSuccess = Result.IsSet();
+	StartFrame = Result.Get(0);
+}
+
+void URenderGridJob::GetEndFrame(bool& bSuccess, int32& EndFrame) const
+{
+	TOptional<int32> Result = GetEndFrame();
+	bSuccess = Result.IsSet();
+	EndFrame = Result.Get(0);
+}
+
+void URenderGridJob::GetStartTime(bool& bSuccess, double& StartTime) const
+{
+	TOptional<double> Result = GetStartTime();
+	bSuccess = Result.IsSet();
+	StartTime = Result.Get(0);
+}
+
+void URenderGridJob::GetEndTime(bool& bSuccess, double& EndTime) const
+{
+	TOptional<double> Result = GetEndTime();
+	bSuccess = Result.IsSet();
+	EndTime = Result.Get(0);
+}
+
+void URenderGridJob::GetDuration(bool& bSuccess, double& Duration) const
+{
+	TOptional<double> Result = GetDuration();
+	bSuccess = Result.IsSet();
+	Duration = Result.Get(0);
+}
+
 FIntPoint URenderGridJob::GetOutputResolution() const
 {
 	if (bOverrideResolution)
@@ -385,7 +491,7 @@ UMoviePipelineOutputSetting* URenderGridJob::GetRenderPresetOutputSettings() con
 	return nullptr;
 }
 
-TArray<URemoteControlPreset*> URenderGridJob::GetRemoteControlPresets()
+TArray<URemoteControlPreset*> URenderGridJob::GetRemoteControlPresets() const
 {
 	TArray<URemoteControlPreset*> Presets;
 	if (URenderGrid* Grid = GetTypedOuter<URenderGrid>(); IsValid(Grid))
@@ -413,28 +519,17 @@ TArray<URemoteControlPreset*> URenderGridJob::GetRemoteControlPresets()
 	return Presets;
 }
 
-bool URenderGridJob::HasRemoteControlValueBytes(const TSharedPtr<FRemoteControlEntity>& RemoteControlEntity) const
-{
-	if (!RemoteControlEntity.IsValid())
-	{
-		return false;
-	}
-
-	return HasRemoteControlValueBytes(RemoteControlEntity->GetId());
-}
-
-bool URenderGridJob::ConstGetRemoteControlValueBytes(const TSharedPtr<FRemoteControlEntity>& RemoteControlEntity, TArray<uint8>& OutBytes) const
+bool URenderGridJob::HasStoredRemoteControlValueBytes(const TSharedPtr<FRemoteControlEntity>& RemoteControlEntity) const
 {
 	if (!RemoteControlEntity.IsValid())
 	{
-		OutBytes.Empty();
 		return false;
 	}
 
-	return ConstGetRemoteControlValueBytes(RemoteControlEntity->GetId(), OutBytes);
+	return HasStoredRemoteControlValueBytes(RemoteControlEntity->GetId());
 }
 
-bool URenderGridJob::GetRemoteControlValueBytes(const TSharedPtr<FRemoteControlEntity>& RemoteControlEntity, TArray<uint8>& OutBytes)
+bool URenderGridJob::GetRemoteControlValueBytes(const TSharedPtr<FRemoteControlEntity>& RemoteControlEntity, TArray<uint8>& OutBytes) const
 {
 	OutBytes.Empty();
 	if (!RemoteControlEntity.IsValid())
@@ -466,7 +561,7 @@ bool URenderGridJob::SetRemoteControlValueBytes(const TSharedPtr<FRemoteControlE
 	return SetRemoteControlValueBytes(RemoteControlEntity->GetId(), Bytes);
 }
 
-bool URenderGridJob::HasRemoteControlValueBytes(const FGuid& FieldId) const
+bool URenderGridJob::HasStoredRemoteControlValueBytes(const FGuid& FieldId) const
 {
 	if (!FieldId.IsValid())
 	{
@@ -476,23 +571,7 @@ bool URenderGridJob::HasRemoteControlValueBytes(const FGuid& FieldId) const
 	return !!RemoteControlValues.Find(FieldId);
 }
 
-bool URenderGridJob::ConstGetRemoteControlValueBytes(const FGuid& FieldId, TArray<uint8>& OutBytes) const
-{
-	OutBytes.Empty();
-	if (!FieldId.IsValid())
-	{
-		return false;
-	}
-
-	if (const FRenderGridRemoteControlPropertyData* DataPtr = RemoteControlValues.Find(FieldId))
-	{
-		OutBytes.Append((*DataPtr).Bytes);
-		return true;
-	}
-	return false;
-}
-
-bool URenderGridJob::GetRemoteControlValueBytes(const FGuid& FieldId, TArray<uint8>& OutBytes)
+bool URenderGridJob::GetRemoteControlValueBytes(const FGuid& FieldId, TArray<uint8>& OutBytes) const
 {
 	OutBytes.Empty();
 	if (!FieldId.IsValid())
@@ -541,24 +620,7 @@ bool URenderGridJob::SetRemoteControlValueBytes(const FGuid& FieldId, const TArr
 	return true;
 }
 
-bool URenderGridJob::HasRemoteControlValue(const FGuid& FieldId) const
-{
-	return HasRemoteControlValueBytes(FieldId);
-}
-
-bool URenderGridJob::ConstGetRemoteControlValue(const FGuid& FieldId, FString& Json) const
-{
-	Json.Empty();
-	TArray<uint8> Bytes;
-	if (!ConstGetRemoteControlValueBytes(FieldId, Bytes))
-	{
-		return false;
-	}
-	Json.Append(UE::RenderGrid::Private::FRenderGridUtils::GetRemoteControlValueJsonFromBytes(Bytes));
-	return true;
-}
-
-bool URenderGridJob::GetRemoteControlValue(const FGuid& FieldId, FString& Json)
+bool URenderGridJob::GetRemoteControlValue(const FGuid& FieldId, FString& Json) const
 {
 	Json.Empty();
 	TArray<uint8> Bytes;
@@ -580,7 +642,7 @@ bool URenderGridJob::SetRemoteControlValue(const FGuid& FieldId, const FString&
 	return SetRemoteControlValueBytes(FieldId, UE::RenderGrid::Private::FRenderGridUtils::GetRemoteControlValueBytesFromJson(Bytes, Json));
 }
 
-bool URenderGridJob::GetRemoteControlFieldIdFromLabel(const FString& Label, FGuid& FieldId)
+bool URenderGridJob::GetRemoteControlFieldIdFromLabel(const FString& Label, FGuid& FieldId) const
 {
 	FieldId.Invalidate();
 	const FName LabelName = FName(Label);
@@ -595,7 +657,7 @@ bool URenderGridJob::GetRemoteControlFieldIdFromLabel(const FString& Label, FGui
 	return false;
 }
 
-bool URenderGridJob::GetRemoteControlLabelFromFieldId(const FGuid& FieldId, FString& Label)
+bool URenderGridJob::GetRemoteControlLabelFromFieldId(const FGuid& FieldId, FString& Label) const
 {
 	Label.Empty();
 	for (URemoteControlPreset* RemoteControlPreset : GetRemoteControlPresets())
@@ -615,15 +677,20 @@ bool URenderGridJob::GetRemoteControlLabelFromFieldId(const FGuid& FieldId, FStr
 	return false;
 }
 
-TMap<FGuid, FString> URenderGridJob::GetRemoteControlValues()
+TMap<FGuid, FString> URenderGridJob::GetRemoteControlValues() const
 {
 	TMap<FGuid, FString> Result;
-	for (TTuple<FGuid, FRenderGridRemoteControlPropertyData> Entry : RemoteControlValues)
+	for (URemoteControlPreset* RemoteControlPreset : GetRemoteControlPresets())
 	{
-		FString Json;
-		if (GetRemoteControlValue(Entry.Key, Json))
+		for (const TWeakPtr<FRemoteControlEntity>& PropWeakPtr : RemoteControlPreset->GetExposedEntities<FRemoteControlEntity>())
 		{
-			Result.Add(Entry.Key, Json);
+			if (const TSharedPtr<FRemoteControlEntity> Prop = PropWeakPtr.Pin())
+			{
+				if (FString Json; GetRemoteControlValue(Prop->GetId(), Json))
+				{
+					Result.Add(Prop->GetId(), Json);
+				}
+			}
 		}
 	}
 	return Result;
@@ -643,6 +710,62 @@ URenderGrid::URenderGrid()
 	}
 }
 
+FString URenderGrid::ToDebugString() const
+{
+	return UE::RenderGrid::Private::FRenderGridUtils::ToJsonString(ToDebugJson(), true);
+}
+
+TSharedPtr<FJsonValue> URenderGrid::ToDebugJson() const
+{
+	TSharedPtr<FJsonObject> Root = MakeShareable(new FJsonObject());
+	{
+		Root->SetStringField(TEXT("Class"), GetParentNativeClass(GetClass())->GetName());
+		Root->SetStringField(TEXT("Object"), GetName());
+
+		{
+			FString Path = GetPathNameSafe(HasAnyFlags(RF_ClassDefaultObject) ? this : GetClass()->GetDefaultObject(true));
+			if (FString PathLeft, PathRight; Path.Split(TEXT("."), &PathLeft, &PathRight, ESearchCase::Type::IgnoreCase, ESearchDir::Type::FromEnd))
+			{
+				FString Name = PathLeft;
+				if (FString NameLeft, NameRight; Name.Split(TEXT("/"), &NameLeft, &NameRight, ESearchCase::Type::IgnoreCase, ESearchDir::Type::FromEnd))
+				{
+					Name = NameRight;
+				}
+				Path = PathLeft + TEXT(".") + Name;
+			}
+			Root->SetStringField(TEXT("Path"), Path);
+		}
+
+		Root->SetStringField(TEXT("Guid"), GetGuid().ToString());
+
+		{
+			TSharedPtr<FJsonObject> Object = MakeShareable(new FJsonObject());
+			Object->SetStringField(TEXT("PropsSourceType"), UEnum::GetDisplayValueAsText(GetPropsSourceType()).ToString());
+			Object->SetStringField(TEXT("PropsSourceOrigin"), GetPathNameSafe(GetPropsSourceOrigin()));
+			Root->SetObjectField(TEXT("Settings"), Object);
+		}
+
+		{
+			TSharedPtr<FJsonObject> Object = MakeShareable(new FJsonObject());
+			Object->SetStringField(TEXT("LevelSequence"), GetPathNameSafe(GetDefaultLevelSequence()));
+			Object->SetStringField(TEXT("RenderPreset"), GetPathNameSafe(GetDefaultRenderPreset()));
+			Object->SetStringField(TEXT("OutputDirectory"), GetDefaultOutputDirectory());
+			Object->SetStringField(TEXT("OutputDirectory"), GetDefaultOutputDirectory());
+			Root->SetObjectField(TEXT("Defaults"), Object);
+		}
+
+		{
+			TArray<TSharedPtr<FJsonValue>> Array;
+			for (URenderGridJob* Job : GetRenderGridJobs())
+			{
+				Array.Add(Job->ToDebugJson());
+			}
+			Root->SetArrayField(TEXT("Jobs"), Array);
+		}
+	}
+	return MakeShareable(new FJsonValueObject(Root));
+}
+
 UWorld* URenderGrid::GetWorld() const
 {
 	for (const FWorldContext& Context : GEngine->GetWorldContexts())
diff --git a/Engine/Plugins/Experimental/RenderGrid/Source/RenderGrid/Private/RenderGrid/RenderGridManager.cpp b/Engine/Plugins/Experimental/RenderGrid/Source/RenderGrid/Private/RenderGrid/RenderGridManager.cpp
--- a/Engine/Plugins/Experimental/RenderGrid/Source/RenderGrid/Private/RenderGrid/RenderGridManager.cpp
+++ b/Engine/Plugins/Experimental/RenderGrid/Source/RenderGrid/Private/RenderGrid/RenderGridManager.cpp
@@ -175,7 +175,7 @@ UTexture2D* UE::RenderGrid::FRenderGridManager::GetRenderedPreviewFrame(URenderG
 }
 
 
-void UE::RenderGrid::FRenderGridManager::UpdateRenderGridJobsPropValues(URenderGrid* Grid)
+void UE::RenderGrid::FRenderGridManager::UpdateStoredRenderGridJobsPropValues(URenderGrid* Grid)
 {
 	if (!IsValid(Grid))
 	{
@@ -200,7 +200,7 @@ void UE::RenderGrid::FRenderGridManager::UpdateRenderGridJobsPropValues(URenderG
 		{
 			for (URenderGridJob* Job : Jobs)
 			{
-				if (!Job->HasRemoteControlValueBytes(Entity))
+				if (!Job->HasStoredRemoteControlValueBytes(Entity))
 				{
 					Job->SetRemoteControlValueBytes(Entity, Bytes);
 				}
@@ -230,7 +230,7 @@ FRenderGridManagerPreviousPropValues UE::RenderGrid::FRenderGridManager::ApplyJo
 			}
 
 			TArray<uint8> PropData;
-			if (!Job->ConstGetRemoteControlValueBytes(Prop->GetRemoteControlEntity(), PropData))
+			if (!Job->GetRemoteControlValueBytes(Prop->GetRemoteControlEntity(), PropData))
 			{
 				continue;
 			}
diff --git a/Engine/Plugins/Experimental/RenderGrid/Source/RenderGrid/Private/RenderGrid/RenderGridQueue.cpp b/Engine/Plugins/Experimental/RenderGrid/Source/RenderGrid/Private/RenderGrid/RenderGridQueue.cpp
--- a/Engine/Plugins/Experimental/RenderGrid/Source/RenderGrid/Private/RenderGrid/RenderGridQueue.cpp
+++ b/Engine/Plugins/Experimental/RenderGrid/Source/RenderGrid/Private/RenderGrid/RenderGridQueue.cpp
@@ -9,6 +9,7 @@
 #include "UObject/Package.h"
 #include "Utils/RenderGridGenericExecutionQueue.h"
 
+#include "Dom/JsonObject.h"
 #include "LevelSequence.h"
 #include "LevelSequenceEditorModule.h"
 #include "Modules/ModuleManager.h"
@@ -51,6 +52,7 @@ URenderGridMoviePipelineRenderJob* URenderGridMoviePipelineRenderJob::Create(URe
 	RenderJob->Status = TEXT("Skipped");
 	RenderJob->bCanExecute = false;
 	RenderJob->bCanceled = false;
+	RenderJob->bFinished = false;
 	RenderJob->Promise = MakeShared<TPromise<void>>();
 	RenderJob->Promise->SetValue();
 	RenderJob->PromiseFuture = RenderJob->Promise->GetFuture().Share();
@@ -197,7 +199,7 @@ URenderGridMoviePipelineRenderJob* URenderGridMoviePipelineRenderJob::Create(URe
 		return RenderJob;
 	}
 
-	RenderJob->Status = TEXT("");
+	RenderJob->Status = TEXT("Waiting");
 	RenderJob->bCanExecute = true;
 	return RenderJob;
 }
@@ -251,6 +253,7 @@ TSharedFuture<void> URenderGridMoviePipelineRenderJob::Execute()
 void URenderGridMoviePipelineRenderJob::Cancel()
 {
 	bCanceled = true;
+	bFinished = true;
 	if (PipelineExecutor->IsRendering())
 	{
 		PipelineExecutor->CancelAllJobs();
@@ -269,6 +272,20 @@ FString URenderGridMoviePipelineRenderJob::GetStatus() const
 	return Status;
 }
 
+float URenderGridMoviePipelineRenderJob::GetStatusPercentage() const
+{
+	if (bCanceled || bFinished || !bCanExecute)
+	{
+		return 100;
+	}
+
+	if (UMoviePipelineExecutorJob* RenderExecutorJob = PipelineExecutorJob.Get(); IsValid(RenderExecutorJob))
+	{
+		return FMath::Clamp(RenderExecutorJob->GetStatusProgress() * 100, 0, 100);
+	}
+	return 0;
+}
+
 int32 URenderGridMoviePipelineRenderJob::GetEngineWarmUpCount() const
 {
 	if (UMoviePipelineExecutorJob* RenderExecutorJob = PipelineExecutorJob.Get(); IsValid(RenderExecutorJob))
@@ -281,6 +298,27 @@ int32 URenderGridMoviePipelineRenderJob::GetEngineWarmUpCount() const
 	return 0;
 }
 
+void URenderGridMoviePipelineRenderJob::OnWaitForEngineWarmUpCount()
+{
+	if (!bCanceled && !bFinished)
+	{
+		Status = TEXT("Warming up...");
+	}
+}
+
+void URenderGridMoviePipelineRenderJob::OnPreRenderEvent()
+{
+	if (!bCanceled && !bFinished)
+	{
+		Status = TEXT("Starting...");
+	}
+}
+
+void URenderGridMoviePipelineRenderJob::OnPostRenderEvent()
+{
+	// the "Done" status is already set by the ExecuteFinished callback at this point
+}
+
 void URenderGridMoviePipelineRenderJob::ComputePlaybackContext(bool& bOutAllowBinding)
 {
 	bOutAllowBinding = false;
@@ -293,6 +331,7 @@ void URenderGridMoviePipelineRenderJob::ExecuteFinished(UMoviePipelineExecutorBa
 		LevelSequenceEditorModule->OnComputePlaybackContext().RemoveAll(this);
 	}
 	bCanceled = (bCanceled || !bSuccess);
+	bFinished = true;
 	Status = (bCanceled ? TEXT("Canceled") : TEXT("Done"));
 	if (Promise.IsValid())
 	{
@@ -303,15 +342,10 @@ void URenderGridMoviePipelineRenderJob::ExecuteFinished(UMoviePipelineExecutorBa
 }
 
 
-TQueue<TObjectPtr<URenderGridQueue>> URenderGridQueue::ExecutingQueues;
-URenderGridQueue::FOnExecutionQueueChanged URenderGridQueue::OnExecutionQueueChangedDelegate;
+TQueue<TObjectPtr<URenderGridQueue>> URenderGridQueue::RenderingQueues;
+URenderGridQueue::FOnRenderingQueueChanged URenderGridQueue::OnRenderingQueueChangedDelegate;
 bool URenderGridQueue::bExitOnCompletion = false;
 
-bool URenderGridQueue::IsExecutingAny()
-{
-	return (GetCurrentlyExecutingQueue() != nullptr);
-}
-
 void URenderGridQueue::CloseEditorOnCompletion()
 {
 	bExitOnCompletion = true;
@@ -320,21 +354,26 @@ void URenderGridQueue::CloseEditorOnCompletion()
 
 void URenderGridQueue::RequestAppExitIfSetToExitOnCompletion()
 {
-	if (bExitOnCompletion && !IsExecutingAny())
+	if (bExitOnCompletion && !IsRenderingAny())
 	{
 		FPlatformMisc::RequestExit(false);
 	}
 }
 
-URenderGridQueue* URenderGridQueue::GetCurrentlyExecutingQueue()
+bool URenderGridQueue::IsRenderingAny()
+{
+	return (GetCurrentlyRenderingQueue() != nullptr);
+}
+
+URenderGridQueue* URenderGridQueue::GetCurrentlyRenderingQueue()
 {
 	bool bRemovedAny = false;
-	while (TObjectPtr<URenderGridQueue>* CurrentlyExecutingQueue = ExecutingQueues.Peek())
+	while (TObjectPtr<URenderGridQueue>* CurrentlyExecutingQueue = RenderingQueues.Peek())
 	{
 		URenderGridQueue* Queue = CurrentlyExecutingQueue->Get();
 		if (!IsValid(Queue))
 		{
-			ExecutingQueues.Pop();
+			RenderingQueues.Pop();
 			bRemovedAny = true;
 			continue;
 		}
@@ -344,37 +383,37 @@ URenderGridQueue* URenderGridQueue::GetCurrentlyExecutingQueue()
 			{
 				Queue->Queue->Start();
 			}
-			OnExecutionQueueChanged().Broadcast();
+			OnRenderingQueueChanged().Broadcast();
 		}
 		return Queue;
 	}
 	return nullptr;
 }
 
-void URenderGridQueue::AddExecutingQueue(URenderGridQueue* Queue)
+void URenderGridQueue::AddRenderingQueue(URenderGridQueue* Queue)
 {
-	ExecutingQueues.Enqueue(Queue);
-	if ((GetCurrentlyExecutingQueue() == Queue) && !Queue->bPaused)
+	RenderingQueues.Enqueue(Queue);
+	if ((GetCurrentlyRenderingQueue() == Queue) && !Queue->bPaused)
 	{
 		Queue->Queue->Start(); // can be called twice here, which is fine (once in GetCurrentlyExecutingQueue() if the previous Queue became invalid, and then once again here)
 	}
-	OnExecutionQueueChanged().Broadcast();
+	OnRenderingQueueChanged().Broadcast();
 }
 
-void URenderGridQueue::DoNextExecutingQueue()
+void URenderGridQueue::DoNextRenderingQueue()
 {
-	ExecutingQueues.Pop();
-	if (URenderGridQueue* Queue = GetCurrentlyExecutingQueue())
+	RenderingQueues.Pop();
+	if (URenderGridQueue* Queue = GetCurrentlyRenderingQueue())
 	{
 		if (!Queue->bPaused)
 		{
 			Queue->Queue->Start(); // can be called twice here, which is fine (once in GetCurrentlyExecutingQueue() if the previous Queue became invalid, and then once again here)
 		}
-		OnExecutionQueueChanged().Broadcast();
+		OnRenderingQueueChanged().Broadcast();
 	}
 	else
 	{
-		OnExecutionQueueChanged().Broadcast();
+		OnRenderingQueueChanged().Broadcast();
 		RequestAppExitIfSetToExitOnCompletion();
 	}
 }
@@ -401,6 +440,7 @@ URenderGridQueue* URenderGridQueue::Create(const UE::RenderGrid::FRenderGridQueu
 	}
 
 	URenderGridQueue* RenderQueue = NewObject<URenderGridQueue>(GetTransientPackage());
+	RenderQueue->Guid = FGuid::NewGuid();
 	RenderQueue->Args = Args;
 	RenderQueue->Queue = MakeShareable(new UE::RenderGrid::Private::FRenderGridGenericExecutionQueue);
 	RenderQueue->RenderGrid = RenderGrid;
@@ -418,6 +458,42 @@ URenderGridQueue* URenderGridQueue::Create(const UE::RenderGrid::FRenderGridQueu
 	return RenderQueue;
 }
 
+FString URenderGridQueue::ToDebugString() const
+{
+	return UE::RenderGrid::Private::FRenderGridUtils::ToJsonString(ToDebugJson(), true);
+}
+
+TSharedPtr<FJsonValue> URenderGridQueue::ToDebugJson() const
+{
+	TSharedPtr<FJsonObject> Root = MakeShareable(new FJsonObject());
+	{
+		Root->SetStringField(TEXT("Class"), GetParentNativeClass(GetClass())->GetName());
+		Root->SetStringField(TEXT("Object"), GetName());
+		Root->SetStringField(TEXT("Guid"), GetGuid().ToString());
+		Root->SetBoolField(TEXT("IsStarted"), IsStarted());
+		Root->SetBoolField(TEXT("IsPaused"), IsPaused());
+		Root->SetBoolField(TEXT("IsCanceled"), IsCanceled());
+		Root->SetBoolField(TEXT("IsFinished"), IsFinished());
+		Root->SetBoolField(TEXT("IsCurrentlyRendering"), IsCurrentlyRendering());
+		Root->SetStringField(TEXT("Status"), GetStatus());
+		Root->SetNumberField(TEXT("StatusPercentage"), GetStatusPercentage());
+
+		{
+			TArray<TSharedPtr<FJsonValue>> Array;
+			for (URenderGridJob* Job : GetJobs())
+			{
+				TSharedPtr<FJsonObject> Object = MakeShareable(new FJsonObject());
+				Object->SetStringField(TEXT("Status"), GetJobStatus(Job));
+				Object->SetNumberField(TEXT("StatusPercentage"), GetJobStatusPercentage(Job));
+				Object->SetField(TEXT("Job"), Job->ToDebugJson());
+				Array.Add(MakeShareable(new FJsonValueObject(Object)));
+			}
+			Root->SetArrayField(TEXT("Jobs"), Array);
+		}
+	}
+	return MakeShareable(new FJsonValueObject(Root));
+}
+
 void URenderGridQueue::Execute()
 {
 	OnStart();
@@ -475,6 +551,11 @@ void URenderGridQueue::Cancel()
 	}
 }
 
+URenderGridJob* URenderGridQueue::GetCurrentlyRenderingJob() const
+{
+	return CurrentJob;
+}
+
 FString URenderGridQueue::GetJobStatus(URenderGridJob* Job) const
 {
 	if (!IsValid(Job))
@@ -492,6 +573,23 @@ FString URenderGridQueue::GetJobStatus(URenderGridJob* Job) const
 	return TEXT("");
 }
 
+float URenderGridQueue::GetJobStatusPercentage(URenderGridJob* Job) const
+{
+	if (!IsValid(Job))
+	{
+		return 0;
+	}
+
+	if (const TObjectPtr<URenderGridMoviePipelineRenderJob>* EntryPtr = Entries.Find(Job))
+	{
+		if (TObjectPtr<URenderGridMoviePipelineRenderJob> Entry = *EntryPtr; IsValid(Entry))
+		{
+			return Entry->GetStatusPercentage();
+		}
+	}
+	return 0;
+}
+
 TArray<URenderGridJob*> URenderGridQueue::GetJobs() const
 {
 	TArray<URenderGridJob*> Jobs;
@@ -552,28 +650,12 @@ int32 URenderGridQueue::GetJobsRemainingCount() const
 			}
 		}
 	}
-	return JobsCount + 1; // one is removed when it starts rendering, but this job is still remaining of course, so let's add 1 here
+	return FMath::Min<int32>(GetJobsCount(), JobsCount + 1); // one is removed when it starts rendering, but this job is still remaining of course, so let's add 1 here
 }
 
 int32 URenderGridQueue::GetJobsCompletedCount() const
 {
-	return (GetJobsCount() - GetJobsRemainingCount());
-}
-
-float URenderGridQueue::GetStatusPercentage() const
-{
-	if (bCanceled || bFinished)
-	{
-		return 100.0;
-	}
-
-	double JobsCount = GetJobsCount();
-	if (JobsCount <= 0)
-	{
-		return 0.0;
-	}
-	double RemainingCount = GetJobsRemainingCount() - 0.5; // rendering a job, but we can't get the progression from the job, so let's say we're halfway there (0.0 started, 0.5 halfway, 1.0 finished)
-	return FMath::Clamp(100.0 - ((RemainingCount / JobsCount) * 100), 0.0, 100.0);
+	return FMath::Max<int32>(0, GetJobsCount() - GetJobsRemainingCount());
 }
 
 FString URenderGridQueue::GetStatus() const
@@ -597,6 +679,26 @@ FString URenderGridQueue::GetStatus() const
 	return TEXT("Rendering...");
 }
 
+float URenderGridQueue::GetStatusPercentage() const
+{
+	if (bCanceled || bFinished)
+	{
+		return 100.0;
+	}
+
+	double JobsCount = GetJobsCount();
+	if (JobsCount <= 0)
+	{
+		return 0.0;
+	}
+	double RemainingCount = GetJobsRemainingCount();
+	if (IsValid(CurrentEntry))
+	{
+		RemainingCount -= (CurrentEntry->GetStatusPercentage() / 100);
+	}
+	return FMath::Clamp(100.0 - ((RemainingCount / JobsCount) * 100), 0.0, 100.0);
+}
+
 
 void URenderGridQueue::OnStart()
 {
@@ -630,7 +732,7 @@ void URenderGridQueue::OnStart()
 
 	OnProcessJob(nullptr); // starts the job execution queueing
 
-	AddExecutingQueue(this);
+	AddRenderingQueue(this);
 }
 
 void URenderGridQueue::OnProcessJob(URenderGridJob* Job)
@@ -641,8 +743,12 @@ void URenderGridQueue::OnProcessJob(URenderGridJob* Job)
 		{
 			if (URenderGridMoviePipelineRenderJob* Entry = *EntryPtr; IsValid(Entry))
 			{
-				Queue->Add(UE::RenderGrid::Private::FRenderGridGenericExecutionQueueActionReturningDelay::CreateLambda([this, Job]() -> UE::RenderGrid::Private::FRenderGridGenericExecutionQueueDelay
+				CurrentJob = Job;
+				CurrentEntry = Entry;
+
+				Queue->Add(UE::RenderGrid::Private::FRenderGridGenericExecutionQueueActionReturningDelay::CreateLambda([this, Job, Entry]() -> UE::RenderGrid::Private::FRenderGridGenericExecutionQueueDelay
 				{
+					Entry->OnWaitForEngineWarmUpCount();
 					if (IsValid(RenderGrid))
 					{
 						PreviousProps = UE::RenderGrid::IRenderGridModule::Get().GetManager().ApplyJobPropValues(RenderGrid, Job);
@@ -650,8 +756,9 @@ void URenderGridQueue::OnProcessJob(URenderGridJob* Job)
 					return UE::RenderGrid::Private::FRenderGridGenericExecutionQueueDelay::Frames(2 + (IsValid(RenderGrid) ? Job->GetWaitFramesBeforeRendering() : 0));
 				}));
 
-				Queue->Add(UE::RenderGrid::Private::FRenderGridGenericExecutionQueueAction::CreateLambda([this, Job]()
+				Queue->Add(UE::RenderGrid::Private::FRenderGridGenericExecutionQueueAction::CreateLambda([this, Job, Entry]()
 				{
+					Entry->OnPreRenderEvent();
 					if (IsValid(RenderGrid) && IsValid(Job))
 					{
 						RenderGrid->BeginJobRender(this, Job);
@@ -663,8 +770,9 @@ void URenderGridQueue::OnProcessJob(URenderGridJob* Job)
 					return Entry->Execute();
 				}));
 
-				Queue->Add(UE::RenderGrid::Private::FRenderGridGenericExecutionQueueAction::CreateLambda([this, Job]()
+				Queue->Add(UE::RenderGrid::Private::FRenderGridGenericExecutionQueueAction::CreateLambda([this, Job, Entry]()
 				{
+					Entry->OnPostRenderEvent();
 					if (IsValid(RenderGrid) && IsValid(Job))
 					{
 						RenderGrid->EndJobRender(this, Job);
@@ -724,9 +832,9 @@ void URenderGridQueue::OnFinish()
 		return;
 	}
 
-	if (GetCurrentlyExecutingQueue() == this)
+	if (GetCurrentlyRenderingQueue() == this)
 	{
-		DoNextExecutingQueue(); // will start the next queue (during the next tick)
+		DoNextRenderingQueue(); // will start the next queue (during the next tick)
 	}
 
 	Queue->Stop();
diff --git a/Engine/Plugins/Experimental/RenderGrid/Source/RenderGrid/Public/RenderGrid/RenderGrid.h b/Engine/Plugins/Experimental/RenderGrid/Source/RenderGrid/Public/RenderGrid/RenderGrid.h
--- a/Engine/Plugins/Experimental/RenderGrid/Source/RenderGrid/Public/RenderGrid/RenderGrid.h
+++ b/Engine/Plugins/Experimental/RenderGrid/Source/RenderGrid/Public/RenderGrid/RenderGrid.h
@@ -7,11 +7,12 @@
 #include "RenderGrid.generated.h"
 
 
-class URenderGridQueue;
+class FJsonValue;
+class ULevelSequence;
 class UMoviePipelineOutputSetting;
 class UMoviePipelinePrimaryConfig;
-class ULevelSequence;
 class URenderGrid;
+class URenderGridQueue;
 
 
 /**
@@ -109,7 +110,7 @@ public:
  * 
  * Each RenderGridJob must belong to a RenderGrid.
  */
-UCLASS(BlueprintType)
+UCLASS(BlueprintType, Meta=(DontUseGenericSpawnObject="true"))
 class RENDERGRID_API URenderGridJob : public UObject
 {
 	GENERATED_BODY()
@@ -117,6 +118,14 @@ class RENDERGRID_API URenderGridJob : public UObject
 public:
 	URenderGridJob();
 
+	/** Obtains a string representation of this object. Shouldn't be used for anything other than logging/debugging. */
+	UFUNCTION(BlueprintPure, Category="Render Grid Job", Meta=(DisplayName="To Debug String (Render Grid Job)", CompactNodeTitle="DEBUG"))
+	FString ToDebugString() const;
+
+	/** Obtains a JSON representation of this object. Shouldn't be used for anything other than logging/debugging. */
+	TSharedPtr<FJsonValue> ToDebugJson() const;
+
+public:
 	/** Gets the calculated start frame, not taking the framerate of the render preset into account. */
 	TOptional<int32> GetSequenceStartFrame() const;
 
@@ -129,20 +138,43 @@ public:
 	/** Sets the custom end frame to match the given sequence end frame. */
 	bool SetSequenceEndFrame(const int32 NewCustomEndFrame);
 
-	/** Gets the calculated start frame. */
+
+	/** Gets the calculated start frame, if possible. */
 	TOptional<int32> GetStartFrame() const;
 
-	/** Gets the calculated end frame. */
+	/** Gets the calculated end frame, if possible. */
 	TOptional<int32> GetEndFrame() const;
 
-	/** Gets the calculated start time. */
+	/** Gets the calculated start time (in seconds), if possible. */
 	TOptional<double> GetStartTime() const;
 
-	/** Gets the calculated end time. */
+	/** Gets the calculated end time (in seconds), if possible. */
 	TOptional<double> GetEndTime() const;
 
-	/** Gets the calculated duration in seconds. */
-	TOptional<double> GetDurationInSeconds() const;
+	/** Gets the calculated duration (in seconds), if possible. */
+	TOptional<double> GetDuration() const;
+
+public:
+	/** Gets the calculated start frame, if possible. */
+	UFUNCTION(BlueprintPure, Category="Render Grid Job")
+	void GetStartFrame(bool& bSuccess, int32& StartFrame) const;
+
+	/** Gets the calculated end frame, if possible. */
+	UFUNCTION(BlueprintPure, Category="Render Grid Job")
+	void GetEndFrame(bool& bSuccess, int32& EndFrame) const;
+
+	/** Gets the calculated start time, if possible. */
+	UFUNCTION(BlueprintPure, Category="Render Grid Job")
+	void GetStartTime(bool& bSuccess, double& StartTime) const;
+
+	/** Gets the calculated end time, if possible. */
+	UFUNCTION(BlueprintPure, Category="Render Grid Job")
+	void GetEndTime(bool& bSuccess, double& EndTime) const;
+
+	/** Gets the calculated duration in seconds, if possible. */
+	UFUNCTION(BlueprintPure, Category="Render Grid Job")
+	void GetDuration(bool& bSuccess, double& Duration) const;
+
 
 	/** Gets the resolution that this job will be rendered in. */
 	UFUNCTION(BlueprintPure, Category="Render Grid Job")
@@ -264,40 +296,33 @@ public:
 	void SetRenderPreset(UMoviePipelinePrimaryConfig* NewRenderPreset) { RenderPreset = NewRenderPreset; }
 
 
-	TArray<URemoteControlPreset*> GetRemoteControlPresets();
+	TArray<URemoteControlPreset*> GetRemoteControlPresets() const;
 
-	bool HasRemoteControlValueBytes(const TSharedPtr<FRemoteControlEntity>& RemoteControlEntity) const;
-	bool ConstGetRemoteControlValueBytes(const TSharedPtr<FRemoteControlEntity>& RemoteControlEntity, TArray<uint8>& OutBytes) const;
-	bool GetRemoteControlValueBytes(const TSharedPtr<FRemoteControlEntity>& RemoteControlEntity, TArray<uint8>& OutBytes);
+	bool HasStoredRemoteControlValueBytes(const TSharedPtr<FRemoteControlEntity>& RemoteControlEntity) const;
+	bool GetRemoteControlValueBytes(const TSharedPtr<FRemoteControlEntity>& RemoteControlEntity, TArray<uint8>& OutBytes) const;
 	bool SetRemoteControlValueBytes(const TSharedPtr<FRemoteControlEntity>& RemoteControlEntity, const TArray<uint8>& Bytes);
 
-	bool HasRemoteControlValueBytes(const FGuid& FieldId) const;
-	bool ConstGetRemoteControlValueBytes(const FGuid& FieldId, TArray<uint8>& OutBytes) const;
-	bool GetRemoteControlValueBytes(const FGuid& FieldId, TArray<uint8>& OutBytes);
+	bool HasStoredRemoteControlValueBytes(const FGuid& FieldId) const;
+	bool GetRemoteControlValueBytes(const FGuid& FieldId, TArray<uint8>& OutBytes) const;
 	bool SetRemoteControlValueBytes(const FGuid& FieldId, const TArray<uint8>& Bytes);
 
 	TMap<FGuid, FRenderGridRemoteControlPropertyData>& GetRemoteControlValuesBytesRef() { return RemoteControlValues; }
 
 
 	UFUNCTION(BlueprintPure, Category="Render Grid Job")
-	bool HasRemoteControlValue(const FGuid& FieldId) const;
-
-	bool ConstGetRemoteControlValue(const FGuid& FieldId, FString& Json) const;
-
-	UFUNCTION(BlueprintPure, Category="Render Grid Job")
-	bool GetRemoteControlValue(const FGuid& FieldId, FString& Json);
+	bool GetRemoteControlValue(const FGuid& FieldId, FString& Json) const;
 
 	UFUNCTION(BlueprintCallable, Category="Render Grid Job")
 	bool SetRemoteControlValue(const FGuid& FieldId, const FString& Json);
 
 	UFUNCTION(BlueprintPure, Category="Render Grid Job")
-	bool GetRemoteControlFieldIdFromLabel(const FString& Label, FGuid& FieldId);
+	bool GetRemoteControlFieldIdFromLabel(const FString& Label, FGuid& FieldId) const;
 
 	UFUNCTION(BlueprintPure, Category="Render Grid Job")
-	bool GetRemoteControlLabelFromFieldId(const FGuid& FieldId, FString& Label);
+	bool GetRemoteControlLabelFromFieldId(const FGuid& FieldId, FString& Label) const;
 
 	UFUNCTION(BlueprintCallable, Category="Render Grid Job")
-	TMap<FGuid, FString> GetRemoteControlValues();
+	TMap<FGuid, FString> GetRemoteControlValues() const;
 
 private:
 	/** The unique ID of this job. */
@@ -358,7 +383,7 @@ private:
 
 	/** The remote control plugin can be used to customize and modify the way a job is rendered. If remote control is being used, the property values of this job will be stored in this map (remote control entity id -> value as bytes). */
 	UPROPERTY()
-	TMap<FGuid, FRenderGridRemoteControlPropertyData> RemoteControlValues;
+	mutable TMap<FGuid, FRenderGridRemoteControlPropertyData> RemoteControlValues;
 };
 
 
@@ -374,6 +399,14 @@ class RENDERGRID_API URenderGrid : public UObject
 public:
 	URenderGrid();
 
+	/** Obtains a string representation of this object. Shouldn't be used for anything other than logging/debugging. */
+	UFUNCTION(BlueprintPure, Category="Render Grid", Meta=(DisplayName="To Debug String (Render Grid)", CompactNodeTitle="DEBUG"))
+	FString ToDebugString() const;
+
+	/** Obtains a JSON representation of this object. Shouldn't be used for anything other than logging/debugging. */
+	TSharedPtr<FJsonValue> ToDebugJson() const;
+
+public:
 	//UObject interface
 	virtual UWorld* GetWorld() const override;
 	virtual void PreSave(FObjectPreSaveContext SaveContext) override;
diff --git a/Engine/Plugins/Experimental/RenderGrid/Source/RenderGrid/Public/RenderGrid/RenderGridManager.h b/Engine/Plugins/Experimental/RenderGrid/Source/RenderGrid/Public/RenderGrid/RenderGridManager.h
--- a/Engine/Plugins/Experimental/RenderGrid/Source/RenderGrid/Public/RenderGrid/RenderGridManager.h
+++ b/Engine/Plugins/Experimental/RenderGrid/Source/RenderGrid/Public/RenderGrid/RenderGridManager.h
@@ -108,7 +108,7 @@ namespace UE::RenderGrid
 
 
 		/** Makes sure that all the data from the current props source is stored in all of the render grid jobs of this render grid. */
-		void UpdateRenderGridJobsPropValues(URenderGrid* Grid);
+		void UpdateStoredRenderGridJobsPropValues(URenderGrid* Grid);
 
 		/** Applies the props of the given render grid job, also requires the render grid to be given as well (to know what props the render grid job is using). */
 		FRenderGridManagerPreviousPropValues ApplyJobPropValues(const URenderGrid* Grid, const URenderGridJob* Job);
diff --git a/Engine/Plugins/Experimental/RenderGrid/Source/RenderGrid/Public/RenderGrid/RenderGridPropsSource.h b/Engine/Plugins/Experimental/RenderGrid/Source/RenderGrid/Public/RenderGrid/RenderGridPropsSource.h
--- a/Engine/Plugins/Experimental/RenderGrid/Source/RenderGrid/Public/RenderGrid/RenderGridPropsSource.h
+++ b/Engine/Plugins/Experimental/RenderGrid/Source/RenderGrid/Public/RenderGrid/RenderGridPropsSource.h
@@ -102,7 +102,7 @@ class RENDERGRID_API URenderGridPropsLocal : public URenderGridPropsBase
 /**
  * The local properties implementation of the render grid properties source abstraction.
  */
-UCLASS(BlueprintType, HideCategories=Object)
+UCLASS(BlueprintType, HideCategories=Object, Meta=(DontUseGenericSpawnObject="true"))
 class RENDERGRID_API URenderGridPropsSourceLocal : public URenderGridPropsSourceBase
 {
 	GENERATED_BODY()
@@ -194,7 +194,7 @@ protected:
 /**
  * The remote control properties implementation of the render grid properties source abstraction.
  */
-UCLASS(BlueprintType, HideCategories=Object)
+UCLASS(BlueprintType, HideCategories=Object, Meta=(DontUseGenericSpawnObject="true"))
 class RENDERGRID_API URenderGridPropsSourceRemoteControl : public URenderGridPropsSourceBase
 {
 	GENERATED_BODY()
diff --git a/Engine/Plugins/Experimental/RenderGrid/Source/RenderGrid/Public/RenderGrid/RenderGridQueue.h b/Engine/Plugins/Experimental/RenderGrid/Source/RenderGrid/Public/RenderGrid/RenderGridQueue.h
--- a/Engine/Plugins/Experimental/RenderGrid/Source/RenderGrid/Public/RenderGrid/RenderGridQueue.h
+++ b/Engine/Plugins/Experimental/RenderGrid/Source/RenderGrid/Public/RenderGrid/RenderGridQueue.h
@@ -10,16 +10,17 @@
 #include "RenderGridQueue.generated.h"
 
 
-class UMoviePipelineSetting;
-class UMoviePipelineOutputBase;
-struct FMoviePipelineOutputData;
-class URenderGrid;
-class URenderGridJob;
+class FJsonValue;
 class UMoviePipelineExecutorBase;
+class UMoviePipelineExecutorJob;
+class UMoviePipelineOutputBase;
 class UMoviePipelinePIEExecutor;
 class UMoviePipelineQueue;
-class UMoviePipelineExecutorJob;
+class UMoviePipelineSetting;
+class URenderGrid;
+class URenderGridJob;
 class URenderGridQueue;
+struct FMoviePipelineOutputData;
 
 namespace UE::RenderGrid::Private
 {
@@ -98,9 +99,22 @@ public:
 	/** Retrieves the rendering status of the given render grid job. */
 	FString GetStatus() const;
 
+	/** Retrieves the rendering status of the given render grid job. */
+	float GetStatusPercentage() const;
+
 	/** Retrieves the "Engine Warm Up Count" value from the AntiAliasingSettings from the render preset that this render grid job uses. */
 	int32 GetEngineWarmUpCount() const;
 
+public:
+	/** An event that will be fired when the rendering process is waiting for this job's engine warm up. */
+	void OnWaitForEngineWarmUpCount();
+
+	/** An event that will be fired when the rendering process is executing the BeginJobRender event for this job. */
+	void OnPreRenderEvent();
+
+	/** An event that will be fired when the rendering process is executing the EndJobRender event for this job. */
+	void OnPostRenderEvent();
+
 public:
 	/** Returns true if this render job was canceled (which for example can be caused by calling Cancel(), or by closing the render popup). */
 	bool IsCanceled() const { return bCanceled; }
@@ -150,50 +164,57 @@ protected:
 	/** Whether the entry was canceled (like by calling Cancel(), or by closing the render popup). */
 	UPROPERTY(Transient)
 	bool bCanceled;
+
+	/** Whether the entry has finished (or was canceled). */
+	UPROPERTY(Transient)
+	bool bFinished;
 };
 
 
 /**
  * This class is responsible for rendering the given render grid jobs.
  */
-UCLASS(BlueprintType)
+UCLASS(BlueprintType, Meta=(DontUseGenericSpawnObject="true"))
 class RENDERGRID_API URenderGridQueue : public UObject, public FTickableGameObject
 {
 	GENERATED_BODY()
 
 public:
-	/** Returns true if it's currently executing any rendering queues. */
-	static bool IsExecutingAny();
-
-	DECLARE_MULTICAST_DELEGATE(FOnExecutionQueueChanged);
-	/** The event that will be fired when the currently executing rendering queue changes (like for example when the previous one has completed its execution). */
-	static FOnExecutionQueueChanged& OnExecutionQueueChanged() { return OnExecutionQueueChangedDelegate; }
+	DECLARE_MULTICAST_DELEGATE(FOnRenderingQueueChanged);
+	/** The event that will be fired when the currently rendering queue changes (like for example when the previous one has finished rendering). */
+	static FOnRenderingQueueChanged& OnRenderingQueueChanged() { return OnRenderingQueueChangedDelegate; }
 
-	/** Call this function to make it so that the editor will be closed when all rendering queues finish execution. This function has to be only called once. */
+	/** Call this function to make it so that the editor will be closed when all rendering queues have finished. This function has to only be called once. */
 	UFUNCTION(BlueprintCallable, Category="Render Grid Queue", Meta=(Keywords="execution execute finish close stop end done quit"))
 	static void CloseEditorOnCompletion();
 
 private:
+	/** Closes the editor, but only if CloseEditorOnCompletion has been called. */
 	static void RequestAppExitIfSetToExitOnCompletion();
 
-private:
-	/** Returns the currently executing queue, or NULL if there isn't any currently executing. */
-	static URenderGridQueue* GetCurrentlyExecutingQueue();
+public:
+	/** Returns true if it's currently rendering any queues. */
+	static bool IsRenderingAny();
 
-	/** Adds the given queue. Will automatically execute it if it's the only one currently. */
-	static void AddExecutingQueue(URenderGridQueue* Queue);
+	/** Returns the currently rendering queue, or NULL if there isn't any currently rendering. */
+	UFUNCTION(BlueprintCallable, Category="Render Grid Queue", Meta=(Keywords="executing execution execute running rendering"))
+	static URenderGridQueue* GetCurrentlyRenderingQueue();
 
-	/** Executes the next queue, if any. Nothing will happen until the next tick. */
-	static void DoNextExecutingQueue();
+private:
+	/** Adds the given queue. Will automatically start rendering it if it's the only one currently. */
+	static void AddRenderingQueue(URenderGridQueue* Queue);
+
+	/** Starts the next rendering queue, if there is one. Nothing will happen until the next tick. */
+	static void DoNextRenderingQueue();
 
 private:
-	/** The queue of executing queues. This contains the currently rendering queue, as well as the ones that should be rendering after it. */
-	static TQueue<TObjectPtr<URenderGridQueue>> ExecutingQueues;
+	/** The queue of rendering queues. This contains the currently rendering queue, as well as the ones that should be rendered after it. */
+	static TQueue<TObjectPtr<URenderGridQueue>> RenderingQueues;
 
 	/** The delegate for when data in the ExecutingQueues has changed. */
-	static FOnExecutionQueueChanged OnExecutionQueueChangedDelegate;
+	static FOnRenderingQueueChanged OnRenderingQueueChangedDelegate;
 
-	/** Whether the editor should be closed when all rendering queues finish their execution. */
+	/** Whether the editor should be closed when all rendering queues finish rendering. */
 	static bool bExitOnCompletion;
 
 public:
@@ -210,13 +231,29 @@ public:
 	}
 	//~ End FTickableGameObject interface
 
+public:
 	/** Creates a new render queue, it won't be started right away. */
 	static URenderGridQueue* Create(const UE::RenderGrid::FRenderGridQueueCreateArgs& Args);
 
+	/** Obtains a string representation of this object. Shouldn't be used for anything other than logging/debugging. */
+	UFUNCTION(BlueprintPure, Category="Render Grid Queue", Meta=(DisplayName="To Debug String (Render Grid Queue)", CompactNodeTitle="DEBUG"))
+	FString ToDebugString() const;
+
+	/** Obtains a JSON representation of this object. Shouldn't be used for anything other than logging/debugging. */
+	TSharedPtr<FJsonValue> ToDebugJson() const;
+
 	/** Starts this render queue. */
 	void Execute();
 
 public:
+	/** Returns the GUID, which is randomly generated at creation. */
+	UFUNCTION(BlueprintPure, Category="Render Grid Queue")
+	FGuid GetGuid() const { return Guid; }
+
+	/** Randomly generates a new GUID. */
+	UFUNCTION(BlueprintCallable, Category="Render Grid Queue")
+	void GenerateNewGuid() { Guid = FGuid::NewGuid(); }
+
 	/** Queues the given job. */
 	UFUNCTION(BlueprintCallable, Category="Render Grid Queue", Meta=(Keywords="render append"))
 	void AddJob(URenderGridJob* Job);
@@ -251,16 +288,24 @@ public:
 
 	/** Returns true if this queue is the one that's currently rendering, returns false if it hasn't started yet, or if it's waiting in the queue, or if it has finished. */
 	UFUNCTION(BlueprintPure, Category="Render Grid Queue", Meta=(Keywords="stopped quited exited killed terminated ended succeeded completed finished canceled cancelled"))
-	bool IsCurrentlyRendering() const { return (GetCurrentlyExecutingQueue() == this); }
+	bool IsCurrentlyRendering() const { return (GetCurrentlyRenderingQueue() == this); }
 
 	/** Retrieves the rendering status of the given render grid job. */
-	UFUNCTION(BlueprintPure, Category="Render Grid Queue", Meta=(Keywords="rendering obtain text"))
+	UFUNCTION(BlueprintPure, Category="Render Grid Queue", Meta=(Keywords="rendering"))
 	URenderGrid* GetRenderGrid() const { return RenderGrid; }
 
-	/** Retrieves the rendering status of the given render grid job. */
+	/** Retrieves the currently rendering render grid job, can return NULL. */
+	UFUNCTION(BlueprintCallable, Category="Render Grid Queue", Meta=(Keywords="rendering progression obtain"))
+	URenderGridJob* GetCurrentlyRenderingJob() const;
+
+	/** Retrieves the rendering status of the given render grid job. Will return an empty string if this job wasn't found in this queue. */
 	UFUNCTION(BlueprintPure, Category="Render Grid Queue", Meta=(Keywords="rendering progression obtain text"))
 	FString GetJobStatus(URenderGridJob* Job) const;
 
+	/** Returns the percentage of the rendering status of the given render grid job. Will return 0 if this job wasn't found in this queue. */
+	UFUNCTION(BlueprintPure, Category="Render Grid Queue", Meta=(Keywords="rendering progression obtain number amount finished complete completion"))
+	float GetJobStatusPercentage(URenderGridJob* Job) const;
+
 	/** Returns all the jobs that have been and will be rendered. */
 	UFUNCTION(BlueprintPure, Category="Render Grid Queue", Meta=(Keywords="rendering progression obtain"))
 	TArray<URenderGridJob*> GetJobs() const;
@@ -269,28 +314,32 @@ public:
 	UFUNCTION(BlueprintPure, Category="Render Grid Queue", Meta=(Keywords="rendering progression obtain number amount"))
 	int32 GetJobsCount() const;
 
-	/** Returns the number of jobs that have finished rendering. Basically just returns [Get Jobs Count] minus [Get Jobs Remaining Count]. */
-	UFUNCTION(BlueprintPure, Category="Render Grid Queue", Meta=(Keywords="rendering progression obtain number amount finished"))
-	int32 GetJobsCompletedCount() const;
-
-	/** Returns the percentage of jobs finished, this includes the progression of the job that is currently rendering. */
-	UFUNCTION(BlueprintPure, Category="Render Grid Queue", Meta=(Keywords="rendering progression obtain number amount finished"))
-	float GetStatusPercentage() const;
-
 	/** Returns the number of jobs that are still left to render, includes the job that is currently rendering. */
 	UFUNCTION(BlueprintPure, Category="Render Grid Queue", Meta=(Keywords="rendering progression obtain number amount finished"))
 	int32 GetJobsRemainingCount() const;
 
+	/** Returns the number of jobs that have finished rendering. Basically just returns [Get Jobs Count] minus [Get Jobs Remaining Count]. */
+	UFUNCTION(BlueprintPure, Category="Render Grid Queue", Meta=(Keywords="rendering progression obtain number amount finished"))
+	int32 GetJobsCompletedCount() const;
+
 	/** Returns the status of the rendering process. */
 	UFUNCTION(BlueprintPure, Category="Render Grid Queue", Meta=(Keywords="rendering progression obtain text"))
 	FString GetStatus() const;
 
+	/** Returns the percentage of jobs finished, this includes the progression of the job that is currently rendering. */
+	UFUNCTION(BlueprintPure, Category="Render Grid Queue", Meta=(Keywords="rendering progression obtain number amount finished complete completion"))
+	float GetStatusPercentage() const;
+
 protected:
 	void OnStart();
 	void OnProcessJob(URenderGridJob* Job);
 	void OnFinish();
 
 protected:
+	/** The unique ID of this queue. */
+	UPROPERTY()
+	FGuid Guid;
+
 	/** The queue containing the render actions. */
 	UE::RenderGrid::FRenderGridQueueCreateArgs Args;
 
@@ -305,6 +354,14 @@ protected:
 	UPROPERTY(Transient)
 	TMap<TObjectPtr<URenderGridJob>, TObjectPtr<URenderGridMoviePipelineRenderJob>> Entries;
 
+	/** The render grid jobs that is currently being rendered. */
+	UPROPERTY(Transient)
+	TObjectPtr<URenderGridJob> CurrentJob;
+
+	/** The movie pipeline render job that is currently being rendered. */
+	UPROPERTY(Transient)
+	TObjectPtr<URenderGridMoviePipelineRenderJob> CurrentEntry;
+
 	/** The render grid of the given render grid job that will be rendered. */
 	UPROPERTY(Transient)
 	TObjectPtr<URenderGrid> RenderGrid;
diff --git a/Engine/Plugins/Experimental/RenderGrid/Source/RenderGrid/RenderGrid.Build.cs b/Engine/Plugins/Experimental/RenderGrid/Source/RenderGrid/RenderGrid.Build.cs
--- a/Engine/Plugins/Experimental/RenderGrid/Source/RenderGrid/RenderGrid.Build.cs
+++ b/Engine/Plugins/Experimental/RenderGrid/Source/RenderGrid/RenderGrid.Build.cs
@@ -7,8 +7,8 @@ public class RenderGrid : ModuleRules
 	public RenderGrid(ReadOnlyTargetRules Target) : base(Target)
 	{
 		PCHUsage = PCHUsageMode.UseExplicitOrSharedPCHs;
-		//PCHUsage = PCHUsageMode.NoPCHs;
-		//bUseUnity = false;
+		PCHUsage = PCHUsageMode.NoPCHs;
+		bUseUnity = false;
 
 		PublicDependencyModuleNames.AddRange(
 			new string[]
diff --git a/Engine/Plugins/Experimental/RenderGrid/Source/RenderGridDeveloper/Public/Blueprints/RenderGridBlueprint.h b/Engine/Plugins/Experimental/RenderGrid/Source/RenderGridDeveloper/Public/Blueprints/RenderGridBlueprint.h
--- a/Engine/Plugins/Experimental/RenderGrid/Source/RenderGridDeveloper/Public/Blueprints/RenderGridBlueprint.h
+++ b/Engine/Plugins/Experimental/RenderGrid/Source/RenderGridDeveloper/Public/Blueprints/RenderGridBlueprint.h
@@ -20,7 +20,7 @@ namespace UE::RenderGrid
  *
  * Required in order for a RenderGrid to be able to have a blueprint graph.
  */
-UCLASS(BlueprintType, Meta=(IgnoreClassThumbnail))
+UCLASS(BlueprintType, Meta=(IgnoreClassThumbnail, DontUseGenericSpawnObject="true"))
 class RENDERGRIDDEVELOPER_API URenderGridBlueprint : public UBlueprint
 {
 	GENERATED_BODY()
diff --git a/Engine/Plugins/Experimental/RenderGrid/Source/RenderGridDeveloper/RenderGridDeveloper.Build.cs b/Engine/Plugins/Experimental/RenderGrid/Source/RenderGridDeveloper/RenderGridDeveloper.Build.cs
--- a/Engine/Plugins/Experimental/RenderGrid/Source/RenderGridDeveloper/RenderGridDeveloper.Build.cs
+++ b/Engine/Plugins/Experimental/RenderGrid/Source/RenderGridDeveloper/RenderGridDeveloper.Build.cs
@@ -7,8 +7,8 @@ public class RenderGridDeveloper : ModuleRules
 	public RenderGridDeveloper(ReadOnlyTargetRules Target) : base(Target)
 	{
 		PCHUsage = PCHUsageMode.UseExplicitOrSharedPCHs;
-		//PCHUsage = PCHUsageMode.NoPCHs;
-		//bUseUnity = false;
+		PCHUsage = PCHUsageMode.NoPCHs;
+		bUseUnity = false;
 
 		PublicDependencyModuleNames.AddRange(
 			new string[]
diff --git a/Engine/Plugins/Experimental/RenderGrid/Source/RenderGridEditor/Private/AssetDefinition/RenderGridBlueprintAssetDefinition.cpp b/Engine/Plugins/Experimental/RenderGrid/Source/RenderGridEditor/Private/AssetDefinition/RenderGridBlueprintAssetDefinition.cpp
--- a/Engine/Plugins/Experimental/RenderGrid/Source/RenderGridEditor/Private/AssetDefinition/RenderGridBlueprintAssetDefinition.cpp
+++ b/Engine/Plugins/Experimental/RenderGrid/Source/RenderGridEditor/Private/AssetDefinition/RenderGridBlueprintAssetDefinition.cpp
@@ -10,7 +10,7 @@
 
 FText URenderGridBlueprintAssetDefinition::GetAssetDisplayName() const
 {
-	return LOCTEXT("AssetTypeActions_RenderGrid", "Render Grid");
+	return LOCTEXT("AssetDefinition_DisplayName_RenderGrid", "Render Grid");
 }
 
 FLinearColor URenderGridBlueprintAssetDefinition::GetAssetColor() const
diff --git a/Engine/Plugins/Experimental/RenderGrid/Source/RenderGridEditor/Private/UI/SRenderGridJobList.cpp b/Engine/Plugins/Experimental/RenderGrid/Source/RenderGridEditor/Private/UI/SRenderGridJobList.cpp
--- a/Engine/Plugins/Experimental/RenderGrid/Source/RenderGridEditor/Private/UI/SRenderGridJobList.cpp
+++ b/Engine/Plugins/Experimental/RenderGrid/Source/RenderGridEditor/Private/UI/SRenderGridJobList.cpp
@@ -696,7 +696,7 @@ TSharedRef<SWidget> UE::RenderGrid::Private::SRenderGridJobListTableRow::Generat
 		else if (ColumnName == FRenderGridJobListColumns::Duration)
 		{
 			FText Text;
-			if (TOptional<double> Duration = Job->GetDurationInSeconds())
+			if (TOptional<double> Duration = Job->GetDuration())
 			{
 				FTimespan Timespan = FTimespan::FromSeconds(*Duration);
 				int32 Hours = static_cast<int32>(Timespan.GetTotalHours());
diff --git a/Engine/Plugins/Experimental/RenderGrid/Source/RenderGridEditor/Private/UI/SRenderGridPropsRemoteControl.cpp b/Engine/Plugins/Experimental/RenderGrid/Source/RenderGridEditor/Private/UI/SRenderGridPropsRemoteControl.cpp
--- a/Engine/Plugins/Experimental/RenderGrid/Source/RenderGridEditor/Private/UI/SRenderGridPropsRemoteControl.cpp
+++ b/Engine/Plugins/Experimental/RenderGrid/Source/RenderGridEditor/Private/UI/SRenderGridPropsRemoteControl.cpp
@@ -54,7 +54,7 @@ void UE::RenderGrid::Private::SRenderGridPropsRemoteControl::UpdateStoredValuesA
 {
 	if (const TSharedPtr<IRenderGridEditor> BlueprintEditor = BlueprintEditorWeakPtr.Pin())
 	{
-		IRenderGridModule::Get().GetManager().UpdateRenderGridJobsPropValues(BlueprintEditor->GetInstance());
+		IRenderGridModule::Get().GetManager().UpdateStoredRenderGridJobsPropValues(BlueprintEditor->GetInstance());
 		Refresh(bForce);
 	}
 }
diff --git a/Engine/Plugins/Experimental/RenderGrid/Source/RenderGridEditor/Public/IRenderGridEditor.h b/Engine/Plugins/Experimental/RenderGrid/Source/RenderGridEditor/Public/IRenderGridEditor.h
--- a/Engine/Plugins/Experimental/RenderGrid/Source/RenderGridEditor/Public/IRenderGridEditor.h
+++ b/Engine/Plugins/Experimental/RenderGrid/Source/RenderGridEditor/Public/IRenderGridEditor.h
@@ -39,7 +39,7 @@ namespace UE::RenderGrid
 		virtual TSharedPtr<Private::FRenderGridBlueprintEditorToolbar> GetRenderGridToolbarBuilder() = 0;
 
 		/** Returns whether it is currently rendering or playing (so changes in the level and such should be ignored). */
-		virtual bool IsCurrentlyRenderingOrPlaying() const { return IsBatchRendering() || IsPreviewRendering() || URenderGridQueue::IsExecutingAny() || IsValid(GEditor->PlayWorld); }
+		virtual bool IsCurrentlyRenderingOrPlaying() const { return IsBatchRendering() || IsPreviewRendering() || URenderGridQueue::IsRenderingAny() || IsValid(GEditor->PlayWorld); }
 
 		/** Returns whether it can currently render (like a preview render or a batch render) or not. */
 		virtual bool CanCurrentlyRender() const { return !IsCurrentlyRenderingOrPlaying(); }
diff --git a/Engine/Plugins/Experimental/RenderGrid/Source/RenderGridEditor/RenderGridEditor.Build.cs b/Engine/Plugins/Experimental/RenderGrid/Source/RenderGridEditor/RenderGridEditor.Build.cs
--- a/Engine/Plugins/Experimental/RenderGrid/Source/RenderGridEditor/RenderGridEditor.Build.cs
+++ b/Engine/Plugins/Experimental/RenderGrid/Source/RenderGridEditor/RenderGridEditor.Build.cs
@@ -7,8 +7,8 @@ public class RenderGridEditor : ModuleRules
 	public RenderGridEditor(ReadOnlyTargetRules Target) : base(Target)
 	{
 		PCHUsage = PCHUsageMode.UseExplicitOrSharedPCHs;
-		//PCHUsage = PCHUsageMode.NoPCHs;
-		//bUseUnity = false;
+		PCHUsage = PCHUsageMode.NoPCHs;
+		bUseUnity = false;
 
 		PublicDependencyModuleNames.AddRange(
 			new string[]
