RenderGrid plugin:
- Fixed the Static Analysis warnings (hopefully).

diff --git a/Engine/Plugins/Experimental/RenderGrid/Source/RenderGrid/Private/Utils/RenderGridGenericExecutionQueue.cpp b/Engine/Plugins/Experimental/RenderGrid/Source/RenderGrid/Private/Utils/RenderGridGenericExecutionQueue.cpp
--- a/Engine/Plugins/Experimental/RenderGrid/Source/RenderGrid/Private/Utils/RenderGridGenericExecutionQueue.cpp
+++ b/Engine/Plugins/Experimental/RenderGrid/Source/RenderGrid/Private/Utils/RenderGridGenericExecutionQueue.cpp
@@ -80,16 +80,20 @@ bool UE::RenderGrid::Private::FRenderGridGenericExecutionQueue::ExecuteNextDelay
 	{
 		if (TDoubleLinkedList<FRenderGridGenericExecutionQueueDelay>::TDoubleLinkedListNode* DelayNode = QueuedDelays.GetHead())
 		{
-			const FRenderGridGenericExecutionQueueDelay& Delay = DelayNode->GetValue();
-			if (Delay.MinimumFrames > 0)
-			{
-				DelayRemainingFrames = Delay.MinimumFrames;
-			}
-			if (Delay.MinimumSeconds > 0)
-			{
-				DelayRemainingSeconds = Delay.MinimumSeconds;
-			}
+			{// uses Delay before calling RemoveNode on it >>
+				const FRenderGridGenericExecutionQueueDelay& Delay = DelayNode->GetValue();
+				if (Delay.MinimumFrames > 0)
+				{
+					DelayRemainingFrames = Delay.MinimumFrames;
+				}
+				if (Delay.MinimumSeconds > 0)
+				{
+					DelayRemainingSeconds = Delay.MinimumSeconds;
+				}
+			}// uses Delay before calling RemoveNode on it <<
+
 			QueuedDelays.RemoveNode(DelayNode);
+
 			if ((DelayRemainingFrames <= 0) && (DelayRemainingSeconds <= 0))
 			{
 				continue;
@@ -104,24 +108,26 @@ bool UE::RenderGrid::Private::FRenderGridGenericExecutionQueue::ExecuteNextEntry
 {
 	if (TDoubleLinkedList<FRenderGridGenericExecutionQueueEntry>::TDoubleLinkedListNode* EntryNode = QueuedEntries.GetHead())
 	{
-		const FRenderGridGenericExecutionQueueEntry& Entry = EntryNode->GetValue();
+		{// uses Entry before calling RemoveNode on it >>
+			const FRenderGridGenericExecutionQueueEntry& Entry = EntryNode->GetValue();
 
-		Entry.ActionRegular.ExecuteIfBound();
+			Entry.ActionRegular.ExecuteIfBound();
 
-		if (Entry.ActionReturningDelay.IsBound())
-		{
-			QueueDelay(Entry.ActionReturningDelay.Execute());
-		}
+			if (Entry.ActionReturningDelay.IsBound())
+			{
+				QueueDelay(Entry.ActionReturningDelay.Execute());
+			}
 
-		if (Entry.ActionReturningDelayFuture.IsBound())
-		{
-			DelayRemainingFuture = Entry.ActionReturningDelayFuture.Execute();
-		}
+			if (Entry.ActionReturningDelayFuture.IsBound())
+			{
+				DelayRemainingFuture = Entry.ActionReturningDelayFuture.Execute();
+			}
 
-		if (Entry.ActionReturningDelayFutureReturningDelay.IsBound())
-		{
-			DelayRemainingFutureReturningDelay = Entry.ActionReturningDelayFutureReturningDelay.Execute();
-		}
+			if (Entry.ActionReturningDelayFutureReturningDelay.IsBound())
+			{
+				DelayRemainingFutureReturningDelay = Entry.ActionReturningDelayFutureReturningDelay.Execute();
+			}
+		}// uses Entry before calling RemoveNode on it <<
 
 		QueuedEntries.RemoveNode(EntryNode);
 		return true;
diff --git a/Engine/Plugins/Experimental/RenderGrid/Source/RenderGrid/RenderGrid.Build.cs b/Engine/Plugins/Experimental/RenderGrid/Source/RenderGrid/RenderGrid.Build.cs
--- a/Engine/Plugins/Experimental/RenderGrid/Source/RenderGrid/RenderGrid.Build.cs
+++ b/Engine/Plugins/Experimental/RenderGrid/Source/RenderGrid/RenderGrid.Build.cs
@@ -7,8 +7,8 @@ public class RenderGrid : ModuleRules
 	public RenderGrid(ReadOnlyTargetRules Target) : base(Target)
 	{
 		PCHUsage = PCHUsageMode.UseExplicitOrSharedPCHs;
-		PCHUsage = PCHUsageMode.NoPCHs;
-		bUseUnity = false;
+		//PCHUsage = PCHUsageMode.NoPCHs;
+		//bUseUnity = false;
 
 		PublicDependencyModuleNames.AddRange(
 			new string[]
diff --git a/Engine/Plugins/Experimental/RenderGrid/Source/RenderGridDeveloper/RenderGridDeveloper.Build.cs b/Engine/Plugins/Experimental/RenderGrid/Source/RenderGridDeveloper/RenderGridDeveloper.Build.cs
--- a/Engine/Plugins/Experimental/RenderGrid/Source/RenderGridDeveloper/RenderGridDeveloper.Build.cs
+++ b/Engine/Plugins/Experimental/RenderGrid/Source/RenderGridDeveloper/RenderGridDeveloper.Build.cs
@@ -7,8 +7,8 @@ public class RenderGridDeveloper : ModuleRules
 	public RenderGridDeveloper(ReadOnlyTargetRules Target) : base(Target)
 	{
 		PCHUsage = PCHUsageMode.UseExplicitOrSharedPCHs;
-		PCHUsage = PCHUsageMode.NoPCHs;
-		bUseUnity = false;
+		//PCHUsage = PCHUsageMode.NoPCHs;
+		//bUseUnity = false;
 
 		PublicDependencyModuleNames.AddRange(
 			new string[]
diff --git a/Engine/Plugins/Experimental/RenderGrid/Source/RenderGridEditor/RenderGridEditor.Build.cs b/Engine/Plugins/Experimental/RenderGrid/Source/RenderGridEditor/RenderGridEditor.Build.cs
--- a/Engine/Plugins/Experimental/RenderGrid/Source/RenderGridEditor/RenderGridEditor.Build.cs
+++ b/Engine/Plugins/Experimental/RenderGrid/Source/RenderGridEditor/RenderGridEditor.Build.cs
@@ -7,8 +7,8 @@ public class RenderGridEditor : ModuleRules
 	public RenderGridEditor(ReadOnlyTargetRules Target) : base(Target)
 	{
 		PCHUsage = PCHUsageMode.UseExplicitOrSharedPCHs;
-		PCHUsage = PCHUsageMode.NoPCHs;
-		bUseUnity = false;
+		//PCHUsage = PCHUsageMode.NoPCHs;
+		//bUseUnity = false;
 
 		PublicDependencyModuleNames.AddRange(
 			new string[]
