RenderGrid plugin:
- Added multiple python example scripts.
- Made it possible to render out a single frame from a given RenderGrid job.
- Added a blueprint to obtain the currently remaining rendering queues.
- Fixed a bug in the GetRemoteControlFieldIdFromLabel blueprint.
- Fixed a missing file error.

diff --git a/Engine/Plugins/Experimental/RenderGrid/Source/RenderGrid/Private/RenderGrid/RenderGrid.cpp b/Engine/Plugins/Experimental/RenderGrid/Source/RenderGrid/Private/RenderGrid/RenderGrid.cpp
--- a/Engine/Plugins/Experimental/RenderGrid/Source/RenderGrid/Private/RenderGrid/RenderGrid.cpp
+++ b/Engine/Plugins/Experimental/RenderGrid/Source/RenderGrid/Private/RenderGrid/RenderGrid.cpp
@@ -648,7 +648,7 @@ bool URenderGridJob::GetRemoteControlFieldIdFromLabel(const FString& Label, FGui
 	const FName LabelName = FName(Label);
 	for (URemoteControlPreset* RemoteControlPreset : GetRemoteControlPresets())
 	{
-		if (const FGuid RemoteControlFieldId = RemoteControlPreset->GetExposedEntityId(LabelName); FieldId.IsValid())
+		if (const FGuid RemoteControlFieldId = RemoteControlPreset->GetExposedEntityId(LabelName); RemoteControlFieldId.IsValid())
 		{
 			FieldId = RemoteControlFieldId;
 			return true;
@@ -945,6 +945,11 @@ URenderGridQueue* URenderGrid::Render()
 	return RenderJobs(GetEnabledRenderGridJobs());
 }
 
+URenderGridQueue* URenderGrid::RenderJob(URenderGridJob* Job)
+{
+	return RenderJobs({Job});
+}
+
 URenderGridQueue* URenderGrid::RenderJobs(const TArray<URenderGridJob*>& Jobs)
 {
 	if (!HasAnyFlags(RF_ClassDefaultObject))
@@ -973,9 +978,42 @@ URenderGridQueue* URenderGrid::RenderJobs(const TArray<URenderGridJob*>& Jobs)
 	return nullptr;
 }
 
-URenderGridQueue* URenderGrid::RenderJob(URenderGridJob* Job)
+URenderGridQueue* URenderGrid::RenderSingleFrame(const double FramePosition)
 {
-	return RenderJobs({Job});
+	return RenderJobsSingleFrame(GetEnabledRenderGridJobs(), FramePosition);
+}
+
+URenderGridQueue* URenderGrid::RenderJobSingleFrame(URenderGridJob* Job, const double FramePosition)
+{
+	return RenderJobsSingleFrame({Job}, FramePosition);
+}
+
+URenderGridQueue* URenderGrid::RenderJobsSingleFrame(const TArray<URenderGridJob*>& Jobs, const double FramePosition)
+{
+	if (!HasAnyFlags(RF_ClassDefaultObject))
+	{
+		if (URenderGrid* DefaultObject = Cast<URenderGrid>(GetClass()->GetDefaultObject(true)); IsValid(DefaultObject))
+		{
+			CopyAllProperties(DefaultObject);
+			return DefaultObject->RenderJobsSingleFrame(Jobs, FramePosition);
+		}
+		return nullptr;
+	}
+
+	TArray<URenderGridJob*> JobsToRender;
+	for (URenderGridJob* Job : Jobs)
+	{
+		if (IsValid(Job) && HasRenderGridJob(Job))
+		{
+			JobsToRender.Add(Job);
+		}
+	}
+	if (URenderGridQueue* RenderQueue = UE::RenderGrid::IRenderGridModule::Get().GetManager().CreateBatchRenderQueueSingleFrame(this, Jobs, FramePosition); IsValid(RenderQueue))
+	{
+		RenderQueue->Execute();
+		return RenderQueue;
+	}
+	return nullptr;
 }
 
 void URenderGrid::SetPropsSource(ERenderGridPropsSourceType InPropsSourceType, UObject* InPropsSourceOrigin)
diff --git a/Engine/Plugins/Experimental/RenderGrid/Source/RenderGrid/Private/RenderGrid/RenderGridManager.cpp b/Engine/Plugins/Experimental/RenderGrid/Source/RenderGrid/Private/RenderGrid/RenderGridManager.cpp
--- a/Engine/Plugins/Experimental/RenderGrid/Source/RenderGrid/Private/RenderGrid/RenderGridManager.cpp
+++ b/Engine/Plugins/Experimental/RenderGrid/Source/RenderGrid/Private/RenderGrid/RenderGridManager.cpp
@@ -20,6 +20,22 @@ URenderGridQueue* UE::RenderGrid::FRenderGridManager::CreateBatchRenderQueue(URe
 	return NewRenderQueue;
 }
 
+URenderGridQueue* UE::RenderGrid::FRenderGridManager::CreateBatchRenderQueueSingleFrame(URenderGrid* Grid, const TArray<URenderGridJob*>& Jobs, const double FramePosition)
+{
+	FRenderGridQueueCreateArgs JobArgs;
+	JobArgs.RenderGrid = TStrongObjectPtr(Grid);
+	JobArgs.RenderGridJobs.Append(Jobs);
+	JobArgs.bIsBatchRender = true;
+	JobArgs.bForceUseSequenceFrameRate = true;
+	JobArgs.FramePosition = FramePosition;
+	URenderGridQueue* NewRenderQueue = URenderGridQueue::Create(JobArgs);
+	if (!IsValid(NewRenderQueue))
+	{
+		return nullptr;
+	}
+	return NewRenderQueue;
+}
+
 
 URenderGridQueue* UE::RenderGrid::FRenderGridManager::RenderPreviewFrame(const FRenderGridManagerRenderPreviewFrameArgs& Args)
 {
diff --git a/Engine/Plugins/Experimental/RenderGrid/Source/RenderGrid/Private/RenderGrid/RenderGridQueue.cpp b/Engine/Plugins/Experimental/RenderGrid/Source/RenderGrid/Private/RenderGrid/RenderGridQueue.cpp
--- a/Engine/Plugins/Experimental/RenderGrid/Source/RenderGrid/Private/RenderGrid/RenderGridQueue.cpp
+++ b/Engine/Plugins/Experimental/RenderGrid/Source/RenderGrid/Private/RenderGrid/RenderGridQueue.cpp
@@ -168,6 +168,14 @@ URenderGridMoviePipelineRenderJob* URenderGridMoviePipelineRenderJob::Create(URe
 		Setting->CustomStartFrame = Job->GetSequenceStartFrame().Get(0);
 		Setting->CustomEndFrame = Job->GetSequenceEndFrame().Get(0);
 
+		if (Args.FramePosition.IsSet())
+		{
+			double FramePosition = FMath::Clamp<double>(Args.FramePosition.Get(0.0), 0.0, 1.0);
+			int32 Frame = FMath::Lerp<int32, double>(Setting->CustomStartFrame, Setting->CustomEndFrame - 1, FramePosition);
+			Setting->CustomStartFrame = Frame;
+			Setting->CustomEndFrame = Frame + 1;
+		}
+
 		if (Args.bForceUseSequenceFrameRate)
 		{
 			Setting->bUseCustomFrameRate = false;
@@ -342,7 +350,7 @@ void URenderGridMoviePipelineRenderJob::ExecuteFinished(UMoviePipelineExecutorBa
 }
 
 
-TQueue<TObjectPtr<URenderGridQueue>> URenderGridQueue::RenderingQueues;
+TArray<TObjectPtr<URenderGridQueue>> URenderGridQueue::RenderingQueues;
 URenderGridQueue::FOnRenderingQueueChanged URenderGridQueue::OnRenderingQueueChangedDelegate;
 bool URenderGridQueue::bExitOnCompletion = false;
 
@@ -368,12 +376,12 @@ bool URenderGridQueue::IsRenderingAny()
 URenderGridQueue* URenderGridQueue::GetCurrentlyRenderingQueue()
 {
 	bool bRemovedAny = false;
-	while (TObjectPtr<URenderGridQueue>* CurrentlyExecutingQueue = RenderingQueues.Peek())
+	while (!RenderingQueues.IsEmpty())
 	{
-		URenderGridQueue* Queue = CurrentlyExecutingQueue->Get();
+		URenderGridQueue* Queue = RenderingQueues[0];
 		if (!IsValid(Queue))
 		{
-			RenderingQueues.Pop();
+			RenderingQueues.RemoveAt(0);
 			bRemovedAny = true;
 			continue;
 		}
@@ -390,9 +398,18 @@ URenderGridQueue* URenderGridQueue::GetCurrentlyRenderingQueue()
 	return nullptr;
 }
 
+int32 URenderGridQueue::GetRemainingRenderingQueuesCount()
+{
+	return RenderingQueues.Num();
+}
+
 void URenderGridQueue::AddRenderingQueue(URenderGridQueue* Queue)
 {
-	RenderingQueues.Enqueue(Queue);
+	if (!IsValid(Queue))
+	{
+		return;
+	}
+	RenderingQueues.Add(Queue);
 	if ((GetCurrentlyRenderingQueue() == Queue) && !Queue->bPaused)
 	{
 		Queue->Queue->Start(); // can be called twice here, which is fine (once in GetCurrentlyExecutingQueue() if the previous Queue became invalid, and then once again here)
@@ -402,7 +419,11 @@ void URenderGridQueue::AddRenderingQueue(URenderGridQueue* Queue)
 
 void URenderGridQueue::DoNextRenderingQueue()
 {
-	RenderingQueues.Pop();
+	if (RenderingQueues.IsEmpty())
+	{
+		return;
+	}
+	RenderingQueues.RemoveAt(0);
 	if (URenderGridQueue* Queue = GetCurrentlyRenderingQueue())
 	{
 		if (!Queue->bPaused)
diff --git a/Engine/Plugins/Experimental/RenderGrid/Source/RenderGrid/Public/RenderGrid/RenderGrid.h b/Engine/Plugins/Experimental/RenderGrid/Source/RenderGrid/Public/RenderGrid/RenderGrid.h
--- a/Engine/Plugins/Experimental/RenderGrid/Source/RenderGrid/Public/RenderGrid/RenderGrid.h
+++ b/Engine/Plugins/Experimental/RenderGrid/Source/RenderGrid/Public/RenderGrid/RenderGrid.h
@@ -548,13 +548,25 @@ public:
 	UFUNCTION(BlueprintCallable, Category="Render Grid")
 	URenderGridQueue* Render();
 
+	/** Renders the given job of this render grid. */
+	UFUNCTION(BlueprintCallable, Category="Render Grid")
+	URenderGridQueue* RenderJob(URenderGridJob* Job);
+
 	/** Renders all the given jobs of this render grid. */
 	UFUNCTION(BlueprintCallable, Category="Render Grid")
 	URenderGridQueue* RenderJobs(const TArray<URenderGridJob*>& Jobs);
 
-	/** Renders the given job of this render grid. */
+	/** Renders all the enabled jobs of this render grid. Only renders a single frame of each job. The frame number it renders is based on the given FramePosition (0.0 is the first frame, 1.0 is the last frame, 0.5 is the frame in the middle, etc). */
 	UFUNCTION(BlueprintCallable, Category="Render Grid")
-	URenderGridQueue* RenderJob(URenderGridJob* Job);
+	URenderGridQueue* RenderSingleFrame(const double FramePosition);
+
+	/** Renders the given job of this render grid. Only renders a single frame. The frame number it renders is based on the given FramePosition (0.0 is the first frame, 1.0 is the last frame, 0.5 is the frame in the middle, etc). */
+	UFUNCTION(BlueprintCallable, Category="Render Grid")
+	URenderGridQueue* RenderJobSingleFrame(URenderGridJob* Job, const double FramePosition);
+
+	/** Renders all the given jobs of this render grid. Only renders a single frame of each job. The frame number it renders is based on the given FramePosition (0.0 is the first frame, 1.0 is the last frame, 0.5 is the frame in the middle, etc). */
+	UFUNCTION(BlueprintCallable, Category="Render Grid")
+	URenderGridQueue* RenderJobsSingleFrame(const TArray<URenderGridJob*>& Jobs, const double FramePosition);
 
 public:
 	/** Returns the GUID, which is randomly generated at creation. */
diff --git a/Engine/Plugins/Experimental/RenderGrid/Source/RenderGrid/Public/RenderGrid/RenderGridManager.h b/Engine/Plugins/Experimental/RenderGrid/Source/RenderGrid/Public/RenderGrid/RenderGridManager.h
--- a/Engine/Plugins/Experimental/RenderGrid/Source/RenderGrid/Public/RenderGrid/RenderGridManager.h
+++ b/Engine/Plugins/Experimental/RenderGrid/Source/RenderGrid/Public/RenderGrid/RenderGridManager.h
@@ -85,11 +85,13 @@ namespace UE::RenderGrid
 		/** The number of characters for a generated ID. For example, a value of 4 results in IDs: "0001", "0002", etc. */
 		static constexpr int32 GeneratedIdCharacterLength = 4;
 
-
 	public:
 		/** Batch render the given render grid job(s) of the given render grid. */
 		URenderGridQueue* CreateBatchRenderQueue(URenderGrid* Grid, const TArray<URenderGridJob*>& Jobs);
 
+		/** Batch render the given render grid job(s) of the given render grid. Only renders a single frame of each job. The frame number it renders is based on the given FramePosition (0.0 is the first frame, 1.0 is the last frame, 0.5 is the frame in the middle, etc). */
+		URenderGridQueue* CreateBatchRenderQueueSingleFrame(URenderGrid* Grid, const TArray<URenderGridJob*>& Jobs, const double FramePosition);
+
 
 		/** Render a preview frame (or multiple if no frame number is specified) of the given render grid job. */
 		URenderGridQueue* RenderPreviewFrame(const FRenderGridManagerRenderPreviewFrameArgs& Args);
@@ -116,7 +118,6 @@ namespace UE::RenderGrid
 		/** Restores the props that were previously applied, to the values that they were before. */
 		void RestorePropValues(const FRenderGridManagerPreviousPropValues& PreviousPropValues);
 
-
 	private:
 		/** The map that stores the start frame (of a render) of each rendered render grid job. */
 		TMap<FGuid, int32> StartFrameOfRenders;
diff --git a/Engine/Plugins/Experimental/RenderGrid/Source/RenderGrid/Public/RenderGrid/RenderGridQueue.h b/Engine/Plugins/Experimental/RenderGrid/Source/RenderGrid/Public/RenderGrid/RenderGridQueue.h
--- a/Engine/Plugins/Experimental/RenderGrid/Source/RenderGrid/Public/RenderGrid/RenderGridQueue.h
+++ b/Engine/Plugins/Experimental/RenderGrid/Source/RenderGrid/Public/RenderGrid/RenderGridQueue.h
@@ -5,7 +5,6 @@
 #include "RenderGridManager.h"
 #include "Utils/RenderGridUtils.h"
 #include "Tickable.h"
-#include "Containers/Queue.h"
 #include "UObject/ObjectPtr.h"
 #include "RenderGridQueue.generated.h"
 
@@ -54,6 +53,9 @@ namespace UE::RenderGrid
 		/** The movie pipeline settings classes to disable (things like Anti-Aliasing, High-Res, etc). */
 		TArray<TSubclassOf<UMoviePipelineSetting>> DisablePipelineSettingsClasses;
 
+		/** The specific frame number that will be rendered. */
+		TOptional<double> FramePosition;
+
 		/** Whether it should run the Begin and End Batch Render events or not. */
 		bool bIsBatchRender = false;
 
@@ -200,6 +202,10 @@ public:
 	UFUNCTION(BlueprintCallable, Category="Render Grid Queue", Meta=(Keywords="executing execution execute running rendering"))
 	static URenderGridQueue* GetCurrentlyRenderingQueue();
 
+	/** Returns the number of rendering queues that are currently queued up. This also includes the currently rendering queue. */
+	UFUNCTION(BlueprintPure, Category="Render Grid Queue", Meta=(Keywords="executing execution execute running rendering left todo queued"))
+	static int32 GetRemainingRenderingQueuesCount();
+
 private:
 	/** Adds the given queue. Will automatically start rendering it if it's the only one currently. */
 	static void AddRenderingQueue(URenderGridQueue* Queue);
@@ -209,7 +215,7 @@ private:
 
 private:
 	/** The queue of rendering queues. This contains the currently rendering queue, as well as the ones that should be rendered after it. */
-	static TQueue<TObjectPtr<URenderGridQueue>> RenderingQueues;
+	static TArray<TObjectPtr<URenderGridQueue>> RenderingQueues;
 
 	/** The delegate for when data in the ExecutingQueues has changed. */
 	static FOnRenderingQueueChanged OnRenderingQueueChangedDelegate;
diff --git a/Engine/Plugins/Experimental/RenderGrid/Source/RenderGridEditor/Private/Styles/RenderGridEditorStyle.cpp b/Engine/Plugins/Experimental/RenderGrid/Source/RenderGridEditor/Private/Styles/RenderGridEditorStyle.cpp
--- a/Engine/Plugins/Experimental/RenderGrid/Source/RenderGridEditor/Private/Styles/RenderGridEditorStyle.cpp
+++ b/Engine/Plugins/Experimental/RenderGrid/Source/RenderGridEditor/Private/Styles/RenderGridEditorStyle.cpp
@@ -67,6 +67,11 @@ const FSlateBrush* UE::RenderGrid::Private::FRenderGridEditorStyle::GetBrush(FNa
 #define BORDER_BRUSH(RelativePath, ...) FSlateBorderBrush(Style->RootToContentDir(RelativePath, TEXT(".png")), __VA_ARGS__)
 #define IMAGE_BRUSH_SVG(RelativePath, ...) FSlateVectorImageBrush(Style->RootToContentDir(RelativePath, TEXT(".svg")), __VA_ARGS__)
 
+#define CORE_IMAGE_BRUSH(RelativePath, ...) FSlateImageBrush(Style->RootToCoreContentDir(RelativePath, TEXT(".png")), __VA_ARGS__)
+#define CORE_BOX_BRUSH(RelativePath, ...) FSlateBoxBrush(Style->RootToCoreContentDir(RelativePath, TEXT(".png")), __VA_ARGS__)
+#define CORE_BORDER_BRUSH(RelativePath, ...) FSlateBorderBrush(Style->RootToCoreContentDir(RelativePath, TEXT(".png")), __VA_ARGS__)
+#define CORE_IMAGE_BRUSH_SVG(RelativePath, ...) FSlateVectorImageBrush(Style->RootToCoreContentDir(RelativePath, TEXT(".svg")), __VA_ARGS__)
+
 const FVector2D Icon64x64(64.f, 64.f);
 const FVector2D Icon40x40(40.0f, 40.0f);
 const FVector2D Icon20x20(20.0f, 20.0f);
@@ -100,8 +105,8 @@ TSharedRef<FSlateStyleSet> UE::RenderGrid::Private::FRenderGridEditorStyle::Crea
 	Style->Set("HoverHintOnly",
 		FButtonStyle()
 		.SetNormal(FSlateNoResource())
-		.SetHovered(BOX_BRUSH("Common/ButtonHoverHint", FMargin(4/16.0f), FLinearColor(1,1,1,0.05f)))
-		.SetPressed(BOX_BRUSH("Common/ButtonHoverHint", FMargin(4/16.0f), FLinearColor(1,1,1,0.15f)))
+		.SetHovered(CORE_BOX_BRUSH("Common/ButtonHoverHint", FMargin(4/16.0f), FLinearColor(1,1,1,0.05f)))
+		.SetPressed(CORE_BOX_BRUSH("Common/ButtonHoverHint", FMargin(4/16.0f), FLinearColor(1,1,1,0.15f)))
 		.SetNormalPadding(FMargin())
 		.SetPressedPadding(FMargin())
 	);
@@ -113,3 +118,8 @@ TSharedRef<FSlateStyleSet> UE::RenderGrid::Private::FRenderGridEditorStyle::Crea
 #undef BOX_BRUSH
 #undef BORDER_BRUSH
 #undef IMAGE_BRUSH_SVG
+
+#undef CORE_IMAGE_BRUSH
+#undef CORE_BOX_BRUSH
+#undef CORE_BORDER_BRUSH
+#undef CORE_IMAGE_BRUSH_SVG
