RenderGrid plugin:
- Updated the documentation of each Python example script.
- Added RenderSingleFrame() functions for a set frame, and reworked the current ones for FramePositions (percentual).
- Added Python example scripts for the regular RenderSingleFrame() functions as well, and reworked the old ones to use RenderSingleFramePosition().
- Made the render grid blueprint graph be an editor-time blueprint graph, to support editor-only blueprints.

diff --git a/Engine/Plugins/Experimental/RenderGrid/Source/RenderGrid/Private/RenderGrid/RenderGrid.cpp b/Engine/Plugins/Experimental/RenderGrid/Source/RenderGrid/Private/RenderGrid/RenderGrid.cpp
--- a/Engine/Plugins/Experimental/RenderGrid/Source/RenderGrid/Private/RenderGrid/RenderGrid.cpp
+++ b/Engine/Plugins/Experimental/RenderGrid/Source/RenderGrid/Private/RenderGrid/RenderGrid.cpp
@@ -978,24 +978,24 @@ URenderGridQueue* URenderGrid::RenderJobs(const TArray<URenderGridJob*>& Jobs)
 	return nullptr;
 }
 
-URenderGridQueue* URenderGrid::RenderSingleFrame(const double FramePosition)
+URenderGridQueue* URenderGrid::RenderSingleFrame(const int32 Frame)
 {
-	return RenderJobsSingleFrame(GetEnabledRenderGridJobs(), FramePosition);
+	return RenderJobsSingleFrame(GetEnabledRenderGridJobs(), Frame);
 }
 
-URenderGridQueue* URenderGrid::RenderJobSingleFrame(URenderGridJob* Job, const double FramePosition)
+URenderGridQueue* URenderGrid::RenderJobSingleFrame(URenderGridJob* Job, const int32 Frame)
 {
-	return RenderJobsSingleFrame({Job}, FramePosition);
+	return RenderJobsSingleFrame({Job}, Frame);
 }
 
-URenderGridQueue* URenderGrid::RenderJobsSingleFrame(const TArray<URenderGridJob*>& Jobs, const double FramePosition)
+URenderGridQueue* URenderGrid::RenderJobsSingleFrame(const TArray<URenderGridJob*>& Jobs, const int32 Frame)
 {
 	if (!HasAnyFlags(RF_ClassDefaultObject))
 	{
 		if (URenderGrid* DefaultObject = Cast<URenderGrid>(GetClass()->GetDefaultObject(true)); IsValid(DefaultObject))
 		{
 			CopyAllProperties(DefaultObject);
-			return DefaultObject->RenderJobsSingleFrame(Jobs, FramePosition);
+			return DefaultObject->RenderJobsSingleFrame(Jobs, Frame);
 		}
 		return nullptr;
 	}
@@ -1008,7 +1008,45 @@ URenderGridQueue* URenderGrid::RenderJobsSingleFrame(const TArray<URenderGridJob
 			JobsToRender.Add(Job);
 		}
 	}
-	if (URenderGridQueue* RenderQueue = UE::RenderGrid::IRenderGridModule::Get().GetManager().CreateBatchRenderQueueSingleFrame(this, Jobs, FramePosition); IsValid(RenderQueue))
+	if (URenderGridQueue* RenderQueue = UE::RenderGrid::IRenderGridModule::Get().GetManager().CreateBatchRenderQueueSingleFrame(this, Jobs, Frame); IsValid(RenderQueue))
+	{
+		RenderQueue->Execute();
+		return RenderQueue;
+	}
+	return nullptr;
+}
+
+URenderGridQueue* URenderGrid::RenderSingleFramePosition(const double FramePosition)
+{
+	return RenderJobsSingleFramePosition(GetEnabledRenderGridJobs(), FramePosition);
+}
+
+URenderGridQueue* URenderGrid::RenderJobSingleFramePosition(URenderGridJob* Job, const double FramePosition)
+{
+	return RenderJobsSingleFramePosition({Job}, FramePosition);
+}
+
+URenderGridQueue* URenderGrid::RenderJobsSingleFramePosition(const TArray<URenderGridJob*>& Jobs, const double FramePosition)
+{
+	if (!HasAnyFlags(RF_ClassDefaultObject))
+	{
+		if (URenderGrid* DefaultObject = Cast<URenderGrid>(GetClass()->GetDefaultObject(true)); IsValid(DefaultObject))
+		{
+			CopyAllProperties(DefaultObject);
+			return DefaultObject->RenderJobsSingleFramePosition(Jobs, FramePosition);
+		}
+		return nullptr;
+	}
+
+	TArray<URenderGridJob*> JobsToRender;
+	for (URenderGridJob* Job : Jobs)
+	{
+		if (IsValid(Job) && HasRenderGridJob(Job))
+		{
+			JobsToRender.Add(Job);
+		}
+	}
+	if (URenderGridQueue* RenderQueue = UE::RenderGrid::IRenderGridModule::Get().GetManager().CreateBatchRenderQueueSingleFramePosition(this, Jobs, FramePosition); IsValid(RenderQueue))
 	{
 		RenderQueue->Execute();
 		return RenderQueue;
diff --git a/Engine/Plugins/Experimental/RenderGrid/Source/RenderGrid/Private/RenderGrid/RenderGridManager.cpp b/Engine/Plugins/Experimental/RenderGrid/Source/RenderGrid/Private/RenderGrid/RenderGridManager.cpp
--- a/Engine/Plugins/Experimental/RenderGrid/Source/RenderGrid/Private/RenderGrid/RenderGridManager.cpp
+++ b/Engine/Plugins/Experimental/RenderGrid/Source/RenderGrid/Private/RenderGrid/RenderGridManager.cpp
@@ -20,7 +20,23 @@ URenderGridQueue* UE::RenderGrid::FRenderGridManager::CreateBatchRenderQueue(URe
 	return NewRenderQueue;
 }
 
-URenderGridQueue* UE::RenderGrid::FRenderGridManager::CreateBatchRenderQueueSingleFrame(URenderGrid* Grid, const TArray<URenderGridJob*>& Jobs, const double FramePosition)
+URenderGridQueue* UE::RenderGrid::FRenderGridManager::CreateBatchRenderQueueSingleFrame(URenderGrid* Grid, const TArray<URenderGridJob*>& Jobs, const int32 Frame)
+{
+	FRenderGridQueueCreateArgs JobArgs;
+	JobArgs.RenderGrid = TStrongObjectPtr(Grid);
+	JobArgs.RenderGridJobs.Append(Jobs);
+	JobArgs.bIsBatchRender = true;
+	JobArgs.bForceUseSequenceFrameRate = true;
+	JobArgs.Frame = Frame;
+	URenderGridQueue* NewRenderQueue = URenderGridQueue::Create(JobArgs);
+	if (!IsValid(NewRenderQueue))
+	{
+		return nullptr;
+	}
+	return NewRenderQueue;
+}
+
+URenderGridQueue* UE::RenderGrid::FRenderGridManager::CreateBatchRenderQueueSingleFramePosition(URenderGrid* Grid, const TArray<URenderGridJob*>& Jobs, const double FramePosition)
 {
 	FRenderGridQueueCreateArgs JobArgs;
 	JobArgs.RenderGrid = TStrongObjectPtr(Grid);
diff --git a/Engine/Plugins/Experimental/RenderGrid/Source/RenderGrid/Private/RenderGrid/RenderGridQueue.cpp b/Engine/Plugins/Experimental/RenderGrid/Source/RenderGrid/Private/RenderGrid/RenderGridQueue.cpp
--- a/Engine/Plugins/Experimental/RenderGrid/Source/RenderGrid/Private/RenderGrid/RenderGridQueue.cpp
+++ b/Engine/Plugins/Experimental/RenderGrid/Source/RenderGrid/Private/RenderGrid/RenderGridQueue.cpp
@@ -168,7 +168,12 @@ URenderGridMoviePipelineRenderJob* URenderGridMoviePipelineRenderJob::Create(URe
 		Setting->CustomStartFrame = Job->GetSequenceStartFrame().Get(0);
 		Setting->CustomEndFrame = Job->GetSequenceEndFrame().Get(0);
 
-		if (Args.FramePosition.IsSet())
+		if (Args.Frame.IsSet())
+		{
+			Setting->CustomStartFrame = Args.Frame.Get(0);
+			Setting->CustomEndFrame = Setting->CustomStartFrame + 1;
+		}
+		else if (Args.FramePosition.IsSet())
 		{
 			double FramePosition = FMath::Clamp<double>(Args.FramePosition.Get(0.0), 0.0, 1.0);
 			int32 Frame = FMath::Lerp<int32, double>(Setting->CustomStartFrame, Setting->CustomEndFrame - 1, FramePosition);
diff --git a/Engine/Plugins/Experimental/RenderGrid/Source/RenderGrid/Public/RenderGrid/RenderGrid.h b/Engine/Plugins/Experimental/RenderGrid/Source/RenderGrid/Public/RenderGrid/RenderGrid.h
--- a/Engine/Plugins/Experimental/RenderGrid/Source/RenderGrid/Public/RenderGrid/RenderGrid.h
+++ b/Engine/Plugins/Experimental/RenderGrid/Source/RenderGrid/Public/RenderGrid/RenderGrid.h
@@ -556,17 +556,31 @@ public:
 	UFUNCTION(BlueprintCallable, Category="Render Grid")
 	URenderGridQueue* RenderJobs(const TArray<URenderGridJob*>& Jobs);
 
+
+	/** Renders all the enabled jobs of this render grid. Only renders a single frame of each job. */
+	UFUNCTION(BlueprintCallable, Category="Render Grid")
+	URenderGridQueue* RenderSingleFrame(const int32 Frame);
+
+	/** Renders the given job of this render grid. Only renders a single frame. */
+	UFUNCTION(BlueprintCallable, Category="Render Grid")
+	URenderGridQueue* RenderJobSingleFrame(URenderGridJob* Job, const int32 Frame);
+
+	/** Renders all the given jobs of this render grid. Only renders a single frame of each job. */
+	UFUNCTION(BlueprintCallable, Category="Render Grid")
+	URenderGridQueue* RenderJobsSingleFrame(const TArray<URenderGridJob*>& Jobs, const int32 Frame);
+
+
 	/** Renders all the enabled jobs of this render grid. Only renders a single frame of each job. The frame number it renders is based on the given FramePosition (0.0 is the first frame, 1.0 is the last frame, 0.5 is the frame in the middle, etc). */
 	UFUNCTION(BlueprintCallable, Category="Render Grid")
-	URenderGridQueue* RenderSingleFrame(const double FramePosition);
+	URenderGridQueue* RenderSingleFramePosition(const double FramePosition);
 
 	/** Renders the given job of this render grid. Only renders a single frame. The frame number it renders is based on the given FramePosition (0.0 is the first frame, 1.0 is the last frame, 0.5 is the frame in the middle, etc). */
 	UFUNCTION(BlueprintCallable, Category="Render Grid")
-	URenderGridQueue* RenderJobSingleFrame(URenderGridJob* Job, const double FramePosition);
+	URenderGridQueue* RenderJobSingleFramePosition(URenderGridJob* Job, const double FramePosition);
 
 	/** Renders all the given jobs of this render grid. Only renders a single frame of each job. The frame number it renders is based on the given FramePosition (0.0 is the first frame, 1.0 is the last frame, 0.5 is the frame in the middle, etc). */
 	UFUNCTION(BlueprintCallable, Category="Render Grid")
-	URenderGridQueue* RenderJobsSingleFrame(const TArray<URenderGridJob*>& Jobs, const double FramePosition);
+	URenderGridQueue* RenderJobsSingleFramePosition(const TArray<URenderGridJob*>& Jobs, const double FramePosition);
 
 public:
 	/** Returns the GUID, which is randomly generated at creation. */
diff --git a/Engine/Plugins/Experimental/RenderGrid/Source/RenderGrid/Public/RenderGrid/RenderGridManager.h b/Engine/Plugins/Experimental/RenderGrid/Source/RenderGrid/Public/RenderGrid/RenderGridManager.h
--- a/Engine/Plugins/Experimental/RenderGrid/Source/RenderGrid/Public/RenderGrid/RenderGridManager.h
+++ b/Engine/Plugins/Experimental/RenderGrid/Source/RenderGrid/Public/RenderGrid/RenderGridManager.h
@@ -89,8 +89,11 @@ namespace UE::RenderGrid
 		/** Batch render the given render grid job(s) of the given render grid. */
 		URenderGridQueue* CreateBatchRenderQueue(URenderGrid* Grid, const TArray<URenderGridJob*>& Jobs);
 
+		/** Batch render the given render grid job(s) of the given render grid. Only renders a single frame of each job. */
+		URenderGridQueue* CreateBatchRenderQueueSingleFrame(URenderGrid* Grid, const TArray<URenderGridJob*>& Jobs, const int32 Frame);
+
 		/** Batch render the given render grid job(s) of the given render grid. Only renders a single frame of each job. The frame number it renders is based on the given FramePosition (0.0 is the first frame, 1.0 is the last frame, 0.5 is the frame in the middle, etc). */
-		URenderGridQueue* CreateBatchRenderQueueSingleFrame(URenderGrid* Grid, const TArray<URenderGridJob*>& Jobs, const double FramePosition);
+		URenderGridQueue* CreateBatchRenderQueueSingleFramePosition(URenderGrid* Grid, const TArray<URenderGridJob*>& Jobs, const double FramePosition);
 
 
 		/** Render a preview frame (or multiple if no frame number is specified) of the given render grid job. */
diff --git a/Engine/Plugins/Experimental/RenderGrid/Source/RenderGrid/Public/RenderGrid/RenderGridQueue.h b/Engine/Plugins/Experimental/RenderGrid/Source/RenderGrid/Public/RenderGrid/RenderGridQueue.h
--- a/Engine/Plugins/Experimental/RenderGrid/Source/RenderGrid/Public/RenderGrid/RenderGridQueue.h
+++ b/Engine/Plugins/Experimental/RenderGrid/Source/RenderGrid/Public/RenderGrid/RenderGridQueue.h
@@ -54,6 +54,9 @@ namespace UE::RenderGrid
 		TArray<TSubclassOf<UMoviePipelineSetting>> DisablePipelineSettingsClasses;
 
 		/** The specific frame number that will be rendered. */
+		TOptional<int32> Frame;
+
+		/** The specific frame number (percentage-wise, between 0.0 and 1.0) that will be rendered. */
 		TOptional<double> FramePosition;
 
 		/** Whether it should run the Begin and End Batch Render events or not. */
diff --git a/Engine/Plugins/Experimental/RenderGrid/Source/RenderGridDeveloper/Public/Blueprints/RenderGridBlueprint.h b/Engine/Plugins/Experimental/RenderGrid/Source/RenderGridDeveloper/Public/Blueprints/RenderGridBlueprint.h
--- a/Engine/Plugins/Experimental/RenderGrid/Source/RenderGridDeveloper/Public/Blueprints/RenderGridBlueprint.h
+++ b/Engine/Plugins/Experimental/RenderGrid/Source/RenderGridDeveloper/Public/Blueprints/RenderGridBlueprint.h
@@ -3,7 +3,7 @@
 #pragma once
 
 #include "CoreMinimal.h"
-#include "Engine/Blueprint.h"
+#include "EditorUtilityBlueprint.h"
 #include "RenderGrid/RenderGrid.h"
 #include "RenderGridBlueprint.generated.h"
 
@@ -21,7 +21,7 @@ namespace UE::RenderGrid
  * Required in order for a RenderGrid to be able to have a blueprint graph.
  */
 UCLASS(BlueprintType, Meta=(IgnoreClassThumbnail, DontUseGenericSpawnObject="true"))
-class RENDERGRIDDEVELOPER_API URenderGridBlueprint : public UBlueprint
+class RENDERGRIDDEVELOPER_API URenderGridBlueprint : public UEditorUtilityBlueprint
 {
 	GENERATED_BODY()
 
@@ -30,6 +30,7 @@ public:
 
 	//~ Begin UBlueprint Interface
 	virtual bool SupportedByDefaultBlueprintFactory() const override { return false; }
+	virtual bool AlwaysCompileOnLoad() const override { return true; }
 	virtual bool IsValidForBytecodeOnlyRecompile() const override { return false; }
 	virtual bool SupportsGlobalVariables() const override { return true; }
 	virtual bool SupportsLocalVariables() const override { return true; }
diff --git a/Engine/Plugins/Experimental/RenderGrid/Source/RenderGridDeveloper/RenderGridDeveloper.Build.cs b/Engine/Plugins/Experimental/RenderGrid/Source/RenderGridDeveloper/RenderGridDeveloper.Build.cs
--- a/Engine/Plugins/Experimental/RenderGrid/Source/RenderGridDeveloper/RenderGridDeveloper.Build.cs
+++ b/Engine/Plugins/Experimental/RenderGrid/Source/RenderGridDeveloper/RenderGridDeveloper.Build.cs
@@ -13,10 +13,11 @@ public class RenderGridDeveloper : ModuleRules
 		PublicDependencyModuleNames.AddRange(
 			new string[]
 			{
+				"Blutility",
 				"Core",
 				"CoreUObject",
 				"Engine",
-				"KismetCompiler",
+				"KismetCompiler", 
 			}
 		);
 
diff --git a/Engine/Plugins/Experimental/RenderGrid/Source/RenderGridEditor/Private/Factories/RenderGridBlueprintFactory.cpp b/Engine/Plugins/Experimental/RenderGrid/Source/RenderGridEditor/Private/Factories/RenderGridBlueprintFactory.cpp
new file mode 100644
--- /dev/null
+++ b/Engine/Plugins/Experimental/RenderGrid/Source/RenderGridEditor/Private/Factories/RenderGridBlueprintFactory.cpp
@@ -0,0 +1,64 @@
+// Copyright Epic Games, Inc. All Rights Reserved.
+
+#include "Factories/RenderGridBlueprintFactory.h"
+#include "Blueprints/RenderGridBlueprint.h"
+#include "IContentBrowserSingleton.h"
+#include "Kismet2/BlueprintEditorUtils.h"
+#include "Kismet2/KismetEditorUtilities.h"
+#include "Misc/MessageDialog.h"
+#include "RenderGrid/RenderGrid.h"
+#include "RenderGrid/RenderGridBlueprintGeneratedClass.h"
+
+#define LOCTEXT_NAMESPACE "RenderGridBlueprintFactory"
+
+
+URenderGridBlueprintFactory::URenderGridBlueprintFactory()
+{
+	ParentClass = URenderGrid::StaticClass();
+	SupportedClass = URenderGridBlueprint::StaticClass();
+	bCreateNew = true; // This factory manufacture new objects from scratch
+	bEditAfterNew = true; // This factory will open the editor for each new object
+}
+
+UObject* URenderGridBlueprintFactory::FactoryCreateNew(UClass* InClass, UObject* InParent, FName InName, EObjectFlags Flags, UObject* Context, FFeedbackContext* Warn, FName CallingContext)
+{
+	// Make sure we are trying to factory a Render Grid Blueprint, then create and init one
+	check(InClass->IsChildOf(URenderGridBlueprint::StaticClass()));
+
+	if ((ParentClass == nullptr) || !FKismetEditorUtilities::CanCreateBlueprintOfClass(ParentClass) || !ParentClass->IsChildOf(URenderGrid::StaticClass()))
+	{
+		FFormatNamedArguments Args;
+		Args.Add(TEXT("ClassName"), (ParentClass != nullptr) ? FText::FromString(ParentClass->GetName()) : LOCTEXT("Null", "(null)"));
+		FMessageDialog::Open(EAppMsgType::Ok, FText::Format(LOCTEXT("CannotCreateRenderGridBlueprint", "Cannot create an Render Grid Blueprint based on the class '{0}'."), Args));
+		return nullptr;
+	}
+
+	if (URenderGridBlueprint* RenderGridBlueprint = CastChecked<URenderGridBlueprint>(FKismetEditorUtilities::CreateBlueprint(ParentClass, InParent, InName, BPTYPE_Normal, URenderGridBlueprint::StaticClass(), URenderGridBlueprintGeneratedClass::StaticClass(), CallingContext)))
+	{
+		RenderGridBlueprint->PostLoad();
+		return RenderGridBlueprint;
+	}
+	return nullptr;
+}
+
+UObject* URenderGridBlueprintFactory::FactoryCreateNew(UClass* InClass, UObject* InParent, FName InName, EObjectFlags Flags, UObject* Context, FFeedbackContext* Warn)
+{
+	return FactoryCreateNew(InClass, InParent, InName, Flags, Context, Warn, NAME_None);
+}
+
+bool URenderGridBlueprintFactory::ShouldShowInNewMenu() const
+{
+	return true;
+}
+
+uint32 URenderGridBlueprintFactory::GetMenuCategories() const
+{
+	//  if wanting to show it in its own category:
+	// IAssetTools& AssetTools = FModuleManager::LoadModuleChecked<FAssetToolsModule>("AssetTools").Get();
+	// return AssetTools.RegisterAdvancedAssetCategory("Render Grid", LOCTEXT("AssetCategoryName", "Render Grid"));
+
+	return EAssetTypeCategories::Misc;
+}
+
+
+#undef LOCTEXT_NAMESPACE
diff --git a/Engine/Plugins/Experimental/RenderGrid/Source/RenderGridEditor/Private/Factories/RenderGridBlueprintFactory.h b/Engine/Plugins/Experimental/RenderGrid/Source/RenderGridEditor/Private/Factories/RenderGridBlueprintFactory.h
new file mode 100644
--- /dev/null
+++ b/Engine/Plugins/Experimental/RenderGrid/Source/RenderGridEditor/Private/Factories/RenderGridBlueprintFactory.h
@@ -0,0 +1,34 @@
+// Copyright Epic Games, Inc. All Rights Reserved.
+
+#pragma once
+
+#include "CoreMinimal.h"
+#include "Factories/Factory.h"
+#include "RenderGrid/RenderGridPropsSource.h"
+#include "RenderGridBlueprintFactory.generated.h"
+
+
+class URenderGridBlueprint;
+class URenderGrid;
+
+
+/**
+ * The factory that creates URenderGridBlueprint (render grid) instances.
+ */
+UCLASS(MinimalAPI, HideCategories=Object)
+class URenderGridBlueprintFactory : public UFactory
+{
+	GENERATED_BODY()
+
+public:
+	URenderGridBlueprintFactory();
+
+	virtual UObject* FactoryCreateNew(UClass* InClass, UObject* InParent, FName InName, EObjectFlags Flags, UObject* Context, FFeedbackContext* Warn, FName CallingContext) override;
+	virtual UObject* FactoryCreateNew(UClass* InClass, UObject* InParent, FName InName, EObjectFlags Flags, UObject* Context, FFeedbackContext* Warn) override;
+	virtual bool ShouldShowInNewMenu() const override;
+	virtual uint32 GetMenuCategories() const override;
+
+	/** The parent class of the created blueprint. */
+	UPROPERTY(EditAnywhere, Category="Render Grid Factory", Meta = (AllowAbstract = ""))
+	TSubclassOf<URenderGrid> ParentClass;
+};
diff --git a/Engine/Plugins/Experimental/RenderGrid/Source/RenderGridEditor/Private/Factories/RenderGridFactory.cpp b/Engine/Plugins/Experimental/RenderGrid/Source/RenderGridEditor/Private/Factories/RenderGridFactory.cpp
deleted file mode 100644
diff --git a/Engine/Plugins/Experimental/RenderGrid/Source/RenderGridEditor/Private/Factories/RenderGridFactory.h b/Engine/Plugins/Experimental/RenderGrid/Source/RenderGridEditor/Private/Factories/RenderGridFactory.h
deleted file mode 100644
