RenderGrid plugin:
- Made non-transient variables in render grid assets work again.
- Fixed several issues that prevented render grid assets from being deleted after having opened them in the editor.

diff --git a/Engine/Plugins/Experimental/RenderGrid/Source/RenderGrid/Private/Factories/RenderGridPropsSourceFactoryLocal.cpp b/Engine/Plugins/Experimental/RenderGrid/Source/RenderGrid/Private/Factories/RenderGridPropsSourceFactoryLocal.cpp
--- a/Engine/Plugins/Experimental/RenderGrid/Source/RenderGrid/Private/Factories/RenderGridPropsSourceFactoryLocal.cpp
+++ b/Engine/Plugins/Experimental/RenderGrid/Source/RenderGrid/Private/Factories/RenderGridPropsSourceFactoryLocal.cpp
@@ -4,9 +4,9 @@
 #include "RenderGrid/RenderGridPropsSource.h"
 
 
-URenderGridPropsSourceBase* UE::RenderGrid::Private::FRenderGridPropsSourceFactoryLocal::CreateInstance(UObject* Outer, UObject* PropsSourceOrigin)
+URenderGridPropsSourceBase* UE::RenderGrid::Private::FRenderGridPropsSourceFactoryLocal::CreateInstance(UObject* PropsSourceOrigin)
 {
-	URenderGridPropsSourceLocal* PropsSource = NewObject<URenderGridPropsSourceLocal>(Outer);
+	URenderGridPropsSourceLocal* PropsSource = NewObject<URenderGridPropsSourceLocal>();
 	PropsSource->SetSourceOrigin(PropsSourceOrigin);
 	return PropsSource;
 }
diff --git a/Engine/Plugins/Experimental/RenderGrid/Source/RenderGrid/Private/Factories/RenderGridPropsSourceFactoryLocal.h b/Engine/Plugins/Experimental/RenderGrid/Source/RenderGrid/Private/Factories/RenderGridPropsSourceFactoryLocal.h
--- a/Engine/Plugins/Experimental/RenderGrid/Source/RenderGrid/Private/Factories/RenderGridPropsSourceFactoryLocal.h
+++ b/Engine/Plugins/Experimental/RenderGrid/Source/RenderGrid/Private/Factories/RenderGridPropsSourceFactoryLocal.h
@@ -14,6 +14,6 @@ namespace UE::RenderGrid::Private
 	class RENDERGRID_API FRenderGridPropsSourceFactoryLocal final : public IRenderGridPropsSourceFactory
 	{
 	public:
-		virtual URenderGridPropsSourceBase* CreateInstance(UObject* Outer, UObject* PropsSourceOrigin) override;
+		virtual URenderGridPropsSourceBase* CreateInstance(UObject* PropsSourceOrigin) override;
 	};
 }
diff --git a/Engine/Plugins/Experimental/RenderGrid/Source/RenderGrid/Private/Factories/RenderGridPropsSourceFactoryRemoteControl.cpp b/Engine/Plugins/Experimental/RenderGrid/Source/RenderGrid/Private/Factories/RenderGridPropsSourceFactoryRemoteControl.cpp
--- a/Engine/Plugins/Experimental/RenderGrid/Source/RenderGrid/Private/Factories/RenderGridPropsSourceFactoryRemoteControl.cpp
+++ b/Engine/Plugins/Experimental/RenderGrid/Source/RenderGrid/Private/Factories/RenderGridPropsSourceFactoryRemoteControl.cpp
@@ -4,9 +4,9 @@
 #include "RenderGrid/RenderGridPropsSource.h"
 
 
-URenderGridPropsSourceBase* UE::RenderGrid::Private::FRenderGridPropsSourceFactoryRemoteControl::CreateInstance(UObject* Outer, UObject* PropsSourceOrigin)
+URenderGridPropsSourceBase* UE::RenderGrid::Private::FRenderGridPropsSourceFactoryRemoteControl::CreateInstance(UObject* PropsSourceOrigin)
 {
-	URenderGridPropsSourceRemoteControl* PropsSource = NewObject<URenderGridPropsSourceRemoteControl>(Outer);
+	URenderGridPropsSourceRemoteControl* PropsSource = NewObject<URenderGridPropsSourceRemoteControl>();
 	PropsSource->SetSourceOrigin(PropsSourceOrigin);
 	return PropsSource;
 }
diff --git a/Engine/Plugins/Experimental/RenderGrid/Source/RenderGrid/Private/Factories/RenderGridPropsSourceFactoryRemoteControl.h b/Engine/Plugins/Experimental/RenderGrid/Source/RenderGrid/Private/Factories/RenderGridPropsSourceFactoryRemoteControl.h
--- a/Engine/Plugins/Experimental/RenderGrid/Source/RenderGrid/Private/Factories/RenderGridPropsSourceFactoryRemoteControl.h
+++ b/Engine/Plugins/Experimental/RenderGrid/Source/RenderGrid/Private/Factories/RenderGridPropsSourceFactoryRemoteControl.h
@@ -14,6 +14,6 @@ namespace UE::RenderGrid::Private
 	class RENDERGRID_API FRenderGridPropsSourceFactoryRemoteControl final : public IRenderGridPropsSourceFactory
 	{
 	public:
-		virtual URenderGridPropsSourceBase* CreateInstance(UObject* Outer, UObject* PropsSourceOrigin) override;
+		virtual URenderGridPropsSourceBase* CreateInstance(UObject* PropsSourceOrigin) override;
 	};
 }
diff --git a/Engine/Plugins/Experimental/RenderGrid/Source/RenderGrid/Private/RenderGrid/RenderGrid.cpp b/Engine/Plugins/Experimental/RenderGrid/Source/RenderGrid/Private/RenderGrid/RenderGrid.cpp
--- a/Engine/Plugins/Experimental/RenderGrid/Source/RenderGrid/Private/RenderGrid/RenderGrid.cpp
+++ b/Engine/Plugins/Experimental/RenderGrid/Source/RenderGrid/Private/RenderGrid/RenderGrid.cpp
@@ -860,7 +860,7 @@ void URenderGrid::CopyAllPropertiesExceptJobs(URenderGrid* From)
 	}
 	for (FProperty* Property = StaticClass()->PropertyLink; Property; Property = Property->PropertyLinkNext)
 	{
-		if (!Property->HasAnyPropertyFlags(CPF_Transient | CPF_DuplicateTransient) && Property->GetName().Compare("RenderGridJobs"))// if not [Transient] and not "RenderGridJobs"
+		if (!Property->HasAnyPropertyFlags(CPF_Transient | CPF_DuplicateTransient) && !Property->GetName().Equals("RenderGridJobs"))// if not [Transient] and not "RenderGridJobs"
 		{
 			Property->CopyCompleteValue(Property->ContainerPtrToValuePtr<void>(this), Property->ContainerPtrToValuePtr<void>(From));
 		}
@@ -877,6 +877,24 @@ void URenderGrid::CopyAllProperties(URenderGrid* From)
 	CopyAllPropertiesExceptJobs(From);
 }
 
+void URenderGrid::CopyAllUserVariables(URenderGrid* From)
+{
+	if (!IsValid(From))
+	{
+		return;
+	}
+	if (From->GetClass()->IsChildOf(GetClass()))
+	{
+		for (FProperty* Property = GetClass()->PropertyLink; Property; Property = Property->PropertyLinkNext)
+		{
+			if (!StaticClass()->HasProperty(Property))// if is [Blueprint Graph Variable]
+			{
+				Property->CopyCompleteValue(Property->ContainerPtrToValuePtr<void>(this), Property->ContainerPtrToValuePtr<void>(From));
+			}
+		}
+	}
+}
+
 void URenderGrid::BeginEditor()
 {
 	bExecutingBlueprintEvent = true;
@@ -1073,7 +1091,7 @@ URenderGridPropsSourceBase* URenderGrid::GetPropsSource() const
 	{
 		Settings->CachedPropsSourceType = Settings->PropsSourceType;
 		Settings->CachedPropsSourceOriginWeakPtr = PropsSourceOrigin;
-		Settings->CachedPropsSource = UE::RenderGrid::IRenderGridModule::Get().CreatePropsSource(const_cast<URenderGrid*>(this), Settings->PropsSourceType, PropsSourceOrigin);
+		Settings->CachedPropsSource = UE::RenderGrid::IRenderGridModule::Get().CreatePropsSource(Settings->PropsSourceType, PropsSourceOrigin);
 	}
 	return Settings->CachedPropsSource;
 }
diff --git a/Engine/Plugins/Experimental/RenderGrid/Source/RenderGrid/Private/RenderGridModule.cpp b/Engine/Plugins/Experimental/RenderGrid/Source/RenderGrid/Private/RenderGridModule.cpp
--- a/Engine/Plugins/Experimental/RenderGrid/Source/RenderGrid/Private/RenderGridModule.cpp
+++ b/Engine/Plugins/Experimental/RenderGrid/Source/RenderGrid/Private/RenderGridModule.cpp
@@ -64,7 +64,7 @@ void UE::RenderGrid::Private::FRenderGridModule::UnregisterPropsSourceFactory(co
 	PropsSourceFactories.Remove(PropsSourceType);
 }
 
-URenderGridPropsSourceBase* UE::RenderGrid::Private::FRenderGridModule::CreatePropsSource(UObject* Outer, ERenderGridPropsSourceType PropsSourceType, UObject* PropsSourceOrigin)
+URenderGridPropsSourceBase* UE::RenderGrid::Private::FRenderGridModule::CreatePropsSource(ERenderGridPropsSourceType PropsSourceType, UObject* PropsSourceOrigin)
 {
 	TSharedPtr<IRenderGridPropsSourceFactory>* FactoryPtr = PropsSourceFactories.Find(PropsSourceType);
 	if (!FactoryPtr)
@@ -78,7 +78,7 @@ URenderGridPropsSourceBase* UE::RenderGrid::Private::FRenderGridModule::CreatePr
 		return nullptr;
 	}
 
-	return Factory->CreateInstance(Outer, PropsSourceOrigin);
+	return Factory->CreateInstance(PropsSourceOrigin);
 }
 
 
diff --git a/Engine/Plugins/Experimental/RenderGrid/Source/RenderGrid/Private/RenderGridModule.h b/Engine/Plugins/Experimental/RenderGrid/Source/RenderGrid/Private/RenderGridModule.h
--- a/Engine/Plugins/Experimental/RenderGrid/Source/RenderGrid/Private/RenderGridModule.h
+++ b/Engine/Plugins/Experimental/RenderGrid/Source/RenderGrid/Private/RenderGridModule.h
@@ -23,7 +23,7 @@ namespace UE::RenderGrid::Private
 		//~ Begin IRenderGridModule interface
 		virtual FRenderGridManager& GetManager() const override;
 
-		virtual URenderGridPropsSourceBase* CreatePropsSource(UObject* Outer, ERenderGridPropsSourceType PropsSourceType, UObject* PropsSourceOrigin) override;
+		virtual URenderGridPropsSourceBase* CreatePropsSource(ERenderGridPropsSourceType PropsSourceType, UObject* PropsSourceOrigin) override;
 		virtual const TMap<ERenderGridPropsSourceType, TSharedPtr<IRenderGridPropsSourceFactory>>& GetPropsSourceFactories() const override { return PropsSourceFactories; }
 		//~ End IRenderGridModule interface
 
diff --git a/Engine/Plugins/Experimental/RenderGrid/Source/RenderGrid/Public/Factories/IRenderGridPropsSourceFactory.h b/Engine/Plugins/Experimental/RenderGrid/Source/RenderGrid/Public/Factories/IRenderGridPropsSourceFactory.h
--- a/Engine/Plugins/Experimental/RenderGrid/Source/RenderGrid/Public/Factories/IRenderGridPropsSourceFactory.h
+++ b/Engine/Plugins/Experimental/RenderGrid/Source/RenderGrid/Public/Factories/IRenderGridPropsSourceFactory.h
@@ -17,6 +17,6 @@ namespace UE::RenderGrid
 	{
 	public:
 		virtual ~IRenderGridPropsSourceFactory() = default;
-		virtual URenderGridPropsSourceBase* CreateInstance(UObject* Outer, UObject* PropsSourceOrigin) { return nullptr; }
+		virtual URenderGridPropsSourceBase* CreateInstance(UObject* PropsSourceOrigin) { return nullptr; }
 	};
 }
diff --git a/Engine/Plugins/Experimental/RenderGrid/Source/RenderGrid/Public/IRenderGridModule.h b/Engine/Plugins/Experimental/RenderGrid/Source/RenderGrid/Public/IRenderGridModule.h
--- a/Engine/Plugins/Experimental/RenderGrid/Source/RenderGrid/Public/IRenderGridModule.h
+++ b/Engine/Plugins/Experimental/RenderGrid/Source/RenderGrid/Public/IRenderGridModule.h
@@ -43,7 +43,7 @@ namespace UE::RenderGrid
 		virtual FRenderGridManager& GetManager() const = 0;
 
 		/** Creates a URenderGridPropsSourceBase instance, based on the given ERenderGridPropsSourceType. */
-		virtual URenderGridPropsSourceBase* CreatePropsSource(UObject* Outer, ERenderGridPropsSourceType PropsSourceType, UObject* PropsSourceOrigin) = 0;
+		virtual URenderGridPropsSourceBase* CreatePropsSource(ERenderGridPropsSourceType PropsSourceType, UObject* PropsSourceOrigin) = 0;
 
 		/** Returns all set IRenderGridPropsSourceFactory instances, these are used to create URenderGridPropsSourceBase instances based on a given ERenderGridPropsSourceType. */
 		virtual const TMap<ERenderGridPropsSourceType, TSharedPtr<IRenderGridPropsSourceFactory>>& GetPropsSourceFactories() const = 0;
diff --git a/Engine/Plugins/Experimental/RenderGrid/Source/RenderGrid/Public/RenderGrid/RenderGrid.h b/Engine/Plugins/Experimental/RenderGrid/Source/RenderGrid/Public/RenderGrid/RenderGrid.h
--- a/Engine/Plugins/Experimental/RenderGrid/Source/RenderGrid/Public/RenderGrid/RenderGrid.h
+++ b/Engine/Plugins/Experimental/RenderGrid/Source/RenderGrid/Public/RenderGrid/RenderGrid.h
@@ -422,6 +422,9 @@ public:
 	/** Copy all properties from the given RenderGrid to self. */
 	void CopyAllProperties(URenderGrid* From);
 
+	/** Copy all user variables (configured in the blueprint graph) from the given RenderGrid to self. */
+	void CopyAllUserVariables(URenderGrid* From);
+
 public:
 	static TArray<FString> GetBlueprintImplementableEvents()
 	{
diff --git a/Engine/Plugins/Experimental/RenderGrid/Source/RenderGridDeveloper/Private/Blueprints/RenderGridBlueprint.cpp b/Engine/Plugins/Experimental/RenderGrid/Source/RenderGridDeveloper/Private/Blueprints/RenderGridBlueprint.cpp
--- a/Engine/Plugins/Experimental/RenderGrid/Source/RenderGridDeveloper/Private/Blueprints/RenderGridBlueprint.cpp
+++ b/Engine/Plugins/Experimental/RenderGrid/Source/RenderGridDeveloper/Private/Blueprints/RenderGridBlueprint.cpp
@@ -89,6 +89,8 @@ void URenderGridBlueprint::Load()
 	{
 		Job->Rename(nullptr, RenderGrid);
 	}
+	RenderGrid->GetDefaultsObject()->Rename(nullptr, RenderGrid);
+	RenderGrid->GetSettingsObject()->Rename(nullptr, RenderGrid);
 }
 
 void URenderGridBlueprint::Save()
@@ -101,6 +103,8 @@ void URenderGridBlueprint::Save()
 			{
 				Job->Rename(nullptr, DefaultObject);
 			}
+			DefaultObject->GetDefaultsObject()->Rename(nullptr, DefaultObject);
+			DefaultObject->GetSettingsObject()->Rename(nullptr, DefaultObject);
 		}
 	}
 }
@@ -132,19 +136,30 @@ void URenderGridBlueprint::PropagateAllPropertiesToInstances()
 void URenderGridBlueprint::PropagateJobsToAsset(URenderGrid* Instance)
 {
 	check(IsValid(Instance));
+
 	RenderGrid->CopyJobs(Instance);
 }
 
 void URenderGridBlueprint::PropagateAllPropertiesExceptJobsToAsset(URenderGrid* Instance)
 {
 	check(IsValid(Instance));
+
 	RenderGrid->CopyAllPropertiesExceptJobs(Instance);
+	if (URenderGrid* CDO = GetRenderGridClassDefaultObject(); IsValid(CDO))
+	{
+		CDO->CopyAllUserVariables(Instance);
+	}
 }
 
 void URenderGridBlueprint::PropagateAllPropertiesToAsset(URenderGrid* Instance)
 {
 	check(IsValid(Instance));
+
 	RenderGrid->CopyAllProperties(Instance);
+	if (URenderGrid* CDO = GetRenderGridClassDefaultObject(); IsValid(CDO))
+	{
+		CDO->CopyAllUserVariables(Instance);
+	}
 }
 
 URenderGrid* URenderGridBlueprint::GetRenderGridWithBlueprintGraph() const
@@ -183,16 +198,7 @@ void URenderGridBlueprint::OnPostVariablesChange(UBlueprint* InBlueprint)
 	for (FBPVariableDescription& NewVariable : NewVariables)
 	{
 		uint64 PreviousPropertyFlags = NewVariable.PropertyFlags;
-		bFoundChange = bFoundChange || NewVariable.HasMetaData(FBlueprintMetadata::MD_ExposeOnSpawn);
-
-		// fix that's required for the way of saving/loading the render grid,
-		//  the render grid instance (in the blueprint class) isn't of the generated class (but instead it's of the URenderGrid class itself), and so variables can't be saved or load,
-		//  meaning that every variable has to be transient no matter what, otherwise saving or loading the render grid asset will cause a crash
-		NewVariable.PropertyFlags |= CPF_DisableEditOnInstance;
-		NewVariable.PropertyFlags &= ~CPF_ExposeOnSpawn;
-		NewVariable.PropertyFlags |= CPF_Transient;
-		NewVariable.RemoveMetaData(FBlueprintMetadata::MD_ExposeOnSpawn);
-
+		NewVariable.PropertyFlags &= ~CPF_Transient;
 		bFoundChange = bFoundChange || (NewVariable.PropertyFlags != PreviousPropertyFlags);
 	}
 
diff --git a/Engine/Plugins/Experimental/RenderGrid/Source/RenderGridEditor/Private/Toolkit/RenderGridEditor.cpp b/Engine/Plugins/Experimental/RenderGrid/Source/RenderGridEditor/Private/Toolkit/RenderGridEditor.cpp
--- a/Engine/Plugins/Experimental/RenderGrid/Source/RenderGridEditor/Private/Toolkit/RenderGridEditor.cpp
+++ b/Engine/Plugins/Experimental/RenderGrid/Source/RenderGridEditor/Private/Toolkit/RenderGridEditor.cpp
@@ -674,9 +674,6 @@ void UE::RenderGrid::Private::FRenderGridEditor::DestroyInstance()
 
 	RenderGridWeakPtr.Reset();
 	PreviewBlueprintWeakPtr.Reset();
-
-	// Clear the undo/redo buffer, to prevent any bugs/glitches with undo/redo while the render grid editor is closed
-	//GEditor->ResetTransaction(LOCTEXT("RenderGridEditorClosed", "RenderGrid Editor Closed"));
 }
 
 void UE::RenderGrid::Private::FRenderGridEditor::UpdateInstance(UBlueprint* InBlueprint, bool bInForceFullUpdate)
