RenderGrid plugin:
- Made instance-editable variables be automatically blueprint-read-only.
- Fixed crashes when trying to undo changes after a rendergrid blueprint compile.

diff --git a/Engine/Plugins/Experimental/RenderGrid/Source/RenderGrid/Private/RenderGrid/RenderGrid.cpp b/Engine/Plugins/Experimental/RenderGrid/Source/RenderGrid/Private/RenderGrid/RenderGrid.cpp
--- a/Engine/Plugins/Experimental/RenderGrid/Source/RenderGrid/Private/RenderGrid/RenderGrid.cpp
+++ b/Engine/Plugins/Experimental/RenderGrid/Source/RenderGrid/Private/RenderGrid/RenderGrid.cpp
@@ -860,11 +860,13 @@ void URenderGrid::CopyAllPropertiesExceptJobs(URenderGrid* From)
 	}
 	for (FProperty* Property = StaticClass()->PropertyLink; Property; Property = Property->PropertyLinkNext)
 	{
-		if (!Property->HasAnyPropertyFlags(CPF_Transient | CPF_DuplicateTransient) && !Property->GetName().Equals("RenderGridJobs"))// if not [Transient] and not "RenderGridJobs"
+		if (!Property->HasAnyPropertyFlags(CPF_Transient | CPF_DuplicateTransient) && !Property->GetName().Equals("RenderGridJobs") && !Property->GetName().Equals("Defaults") && !Property->GetName().Equals("Settings"))// if not [Transient] and not "RenderGridJobs"/"Defaults"/"Settings"
 		{
 			Property->CopyCompleteValue(Property->ContainerPtrToValuePtr<void>(this), Property->ContainerPtrToValuePtr<void>(From));
 		}
 	}
+	Defaults->CopyValuesFrom(From->GetDefaultsObject());
+	Settings->CopyValuesFrom(From->GetSettingsObject());
 }
 
 void URenderGrid::CopyAllProperties(URenderGrid* From)
@@ -975,6 +977,7 @@ URenderGridQueue* URenderGrid::RenderJobs(const TArray<URenderGridJob*>& Jobs)
 		if (URenderGrid* DefaultObject = Cast<URenderGrid>(GetClass()->GetDefaultObject(true)); IsValid(DefaultObject))
 		{
 			CopyAllProperties(DefaultObject);
+			CopyAllUserVariables(DefaultObject);
 			return DefaultObject->RenderJobs(Jobs);
 		}
 		return nullptr;
@@ -1013,6 +1016,7 @@ URenderGridQueue* URenderGrid::RenderJobsSingleFrame(const TArray<URenderGridJob
 		if (URenderGrid* DefaultObject = Cast<URenderGrid>(GetClass()->GetDefaultObject(true)); IsValid(DefaultObject))
 		{
 			CopyAllProperties(DefaultObject);
+			CopyAllUserVariables(DefaultObject);
 			return DefaultObject->RenderJobsSingleFrame(Jobs, Frame);
 		}
 		return nullptr;
@@ -1051,6 +1055,7 @@ URenderGridQueue* URenderGrid::RenderJobsSingleFramePosition(const TArray<URende
 		if (URenderGrid* DefaultObject = Cast<URenderGrid>(GetClass()->GetDefaultObject(true)); IsValid(DefaultObject))
 		{
 			CopyAllProperties(DefaultObject);
+			CopyAllUserVariables(DefaultObject);
 			return DefaultObject->RenderJobsSingleFramePosition(Jobs, FramePosition);
 		}
 		return nullptr;
diff --git a/Engine/Plugins/Experimental/RenderGrid/Source/RenderGrid/Public/RenderGrid/RenderGrid.h b/Engine/Plugins/Experimental/RenderGrid/Source/RenderGrid/Public/RenderGrid/RenderGrid.h
--- a/Engine/Plugins/Experimental/RenderGrid/Source/RenderGrid/Public/RenderGrid/RenderGrid.h
+++ b/Engine/Plugins/Experimental/RenderGrid/Source/RenderGrid/Public/RenderGrid/RenderGrid.h
@@ -52,6 +52,12 @@ class RENDERGRID_API URenderGridSettings : public UObject
 public:
 	URenderGridSettings();
 
+	void CopyValuesFrom(URenderGridSettings* From)
+	{
+		PropsSourceType = From->PropsSourceType;
+		PropsSourceOrigin_RemoteControl = From->PropsSourceOrigin_RemoteControl;
+	}
+
 public:
 	/** The type of the properties that a job in this grid can have. */
 	UPROPERTY(/*EditInstanceOnly, Category="Render Grid", Meta=(DisplayName="Properties Type")*/)
@@ -89,6 +95,13 @@ class RENDERGRID_API URenderGridDefaults : public UObject
 public:
 	URenderGridDefaults();
 
+	void CopyValuesFrom(URenderGridDefaults* From)
+	{
+		LevelSequence = From->LevelSequence;
+		RenderPreset = From->RenderPreset;
+		OutputDirectory = From->OutputDirectory;
+	}
+
 public:
 	/** The default level sequence for new jobs, this is what will be rendered during rendering. A job without a level sequence can't be rendered. */
 	UPROPERTY(EditInstanceOnly, Category="Defaults", Meta=(DisplayName="Default Level Sequence"))
diff --git a/Engine/Plugins/Experimental/RenderGrid/Source/RenderGridDeveloper/Private/Blueprints/RenderGridBlueprint.cpp b/Engine/Plugins/Experimental/RenderGrid/Source/RenderGridDeveloper/Private/Blueprints/RenderGridBlueprint.cpp
--- a/Engine/Plugins/Experimental/RenderGrid/Source/RenderGridDeveloper/Private/Blueprints/RenderGridBlueprint.cpp
+++ b/Engine/Plugins/Experimental/RenderGrid/Source/RenderGridDeveloper/Private/Blueprints/RenderGridBlueprint.cpp
@@ -89,8 +89,6 @@ void URenderGridBlueprint::Load()
 	{
 		Job->Rename(nullptr, RenderGrid);
 	}
-	RenderGrid->GetDefaultsObject()->Rename(nullptr, RenderGrid);
-	RenderGrid->GetSettingsObject()->Rename(nullptr, RenderGrid);
 }
 
 void URenderGridBlueprint::Save()
@@ -103,8 +101,6 @@ void URenderGridBlueprint::Save()
 			{
 				Job->Rename(nullptr, DefaultObject);
 			}
-			DefaultObject->GetDefaultsObject()->Rename(nullptr, DefaultObject);
-			DefaultObject->GetSettingsObject()->Rename(nullptr, DefaultObject);
 		}
 	}
 }
@@ -119,17 +115,27 @@ void URenderGridBlueprint::PropagateJobsToInstances()
 
 void URenderGridBlueprint::PropagateAllPropertiesExceptJobsToInstances()
 {
-	RunOnInstances(UE::RenderGrid::FRenderGridBlueprintRunOnInstancesCallback::CreateLambda([this](URenderGrid* Instance)
+	URenderGrid* CDO = GetRenderGridClassDefaultObject();
+	RunOnInstances(UE::RenderGrid::FRenderGridBlueprintRunOnInstancesCallback::CreateLambda([this, CDO](URenderGrid* Instance)
 	{
 		Instance->CopyAllPropertiesExceptJobs(RenderGrid);
+		if (IsValid(CDO))
+		{
+			Instance->CopyAllUserVariables(CDO);
+		}
 	}));
 }
 
 void URenderGridBlueprint::PropagateAllPropertiesToInstances()
 {
-	RunOnInstances(UE::RenderGrid::FRenderGridBlueprintRunOnInstancesCallback::CreateLambda([this](URenderGrid* Instance)
+	URenderGrid* CDO = GetRenderGridClassDefaultObject();
+	RunOnInstances(UE::RenderGrid::FRenderGridBlueprintRunOnInstancesCallback::CreateLambda([this, CDO](URenderGrid* Instance)
 	{
 		Instance->CopyAllProperties(RenderGrid);
+		if (IsValid(CDO))
+		{
+			Instance->CopyAllUserVariables(CDO);
+		}
 	}));
 }
 
@@ -198,7 +204,16 @@ void URenderGridBlueprint::OnPostVariablesChange(UBlueprint* InBlueprint)
 	for (FBPVariableDescription& NewVariable : NewVariables)
 	{
 		uint64 PreviousPropertyFlags = NewVariable.PropertyFlags;
-		NewVariable.PropertyFlags &= ~CPF_Transient;
+		bFoundChange = bFoundChange || NewVariable.HasMetaData(FBlueprintMetadata::MD_ExposeOnSpawn);
+
+		if ((NewVariable.PropertyFlags & CPF_DisableEditOnInstance) == 0)// if [Instance Editable]
+		{
+			NewVariable.PropertyFlags |= CPF_BlueprintReadOnly;// set [Blueprint Read Only] to true
+		}
+		NewVariable.PropertyFlags &= ~CPF_ExposeOnSpawn;
+		NewVariable.PropertyFlags |= CPF_Transient;
+		NewVariable.RemoveMetaData(FBlueprintMetadata::MD_ExposeOnSpawn);
+
 		bFoundChange = bFoundChange || (NewVariable.PropertyFlags != PreviousPropertyFlags);
 	}
 
