RenderGrid plugin:
- Added a map parameter to the rendergrid asset, and made it correctly handle the cases where that map isn't currently open in the editor.
- Made it possible to click on the viewport to automatically open the rendergrid's map.
- Modified the python scripts, as defining the map (as a command line argument to unreal) is now no longer necessary.

diff --git a/Engine/Plugins/Experimental/RenderGrid/Source/RenderGrid/Private/RenderGrid/RenderGrid.cpp b/Engine/Plugins/Experimental/RenderGrid/Source/RenderGrid/Private/RenderGrid/RenderGrid.cpp
--- a/Engine/Plugins/Experimental/RenderGrid/Source/RenderGrid/Private/RenderGrid/RenderGrid.cpp
+++ b/Engine/Plugins/Experimental/RenderGrid/Source/RenderGrid/Private/RenderGrid/RenderGrid.cpp
@@ -2,7 +2,6 @@
 
 #include "RenderGrid/RenderGrid.h"
 
-#include "DiffUtils.h"
 #include "IRenderGridModule.h"
 #include "LevelSequence.h"
 #include "Misc/Base64.h"
@@ -960,17 +959,7 @@ void URenderGrid::EndViewportRender(URenderGridJob* Job)
 	bExecutingBlueprintEvent = false;
 }
 
-URenderGridQueue* URenderGrid::Render()
-{
-	return RenderJobs(GetEnabledRenderGridJobs());
-}
-
-URenderGridQueue* URenderGrid::RenderJob(URenderGridJob* Job)
-{
-	return RenderJobs({Job});
-}
-
-URenderGridQueue* URenderGrid::RenderJobs(const TArray<URenderGridJob*>& Jobs)
+URenderGridQueue* URenderGrid::RunRenderJobs(const TArray<URenderGridJob*>& Jobs, const FRenderGridRunRenderJobsDefaultObjectCallback& DefaultObjectCallback, const FRenderGridRunRenderJobsCallback& Callback)
 {
 	if (!HasAnyFlags(RF_ClassDefaultObject))
 	{
@@ -978,7 +967,10 @@ URenderGridQueue* URenderGrid::RenderJobs(const TArray<URenderGridJob*>& Jobs)
 		{
 			CopyAllProperties(DefaultObject);
 			CopyAllUserVariables(DefaultObject);
-			return DefaultObject->RenderJobs(Jobs);
+			if (DefaultObjectCallback.IsBound())
+			{
+				return DefaultObjectCallback.Execute(DefaultObject, Jobs);
+			}
 		}
 		return nullptr;
 	}
@@ -991,14 +983,40 @@ URenderGridQueue* URenderGrid::RenderJobs(const TArray<URenderGridJob*>& Jobs)
 			JobsToRender.Add(Job);
 		}
 	}
-	if (URenderGridQueue* RenderQueue = UE::RenderGrid::IRenderGridModule::Get().GetManager().CreateBatchRenderQueue(this, Jobs); IsValid(RenderQueue))
+	if (Callback.IsBound())
 	{
-		RenderQueue->Execute();
-		return RenderQueue;
+		if (URenderGridQueue* RenderQueue = Callback.Execute(this, Jobs); IsValid(RenderQueue))
+		{
+			RenderQueue->Execute();
+			return RenderQueue;
+		}
 	}
 	return nullptr;
 }
 
+URenderGridQueue* URenderGrid::Render()
+{
+	return RenderJobs(GetEnabledRenderGridJobs());
+}
+
+URenderGridQueue* URenderGrid::RenderJob(URenderGridJob* Job)
+{
+	return RenderJobs({Job});
+}
+
+URenderGridQueue* URenderGrid::RenderJobs(const TArray<URenderGridJob*>& Jobs)
+{
+	return RunRenderJobs(
+		Jobs,
+		FRenderGridRunRenderJobsDefaultObjectCallback::CreateLambda([](URenderGrid* InDefaultObject, const TArray<URenderGridJob*>& InJobs) -> URenderGridQueue* {
+			return InDefaultObject->RenderJobs(InJobs);
+		}),
+		FRenderGridRunRenderJobsCallback::CreateLambda([](URenderGrid* InThis, const TArray<URenderGridJob*>& InJobs) -> URenderGridQueue* {
+			return UE::RenderGrid::IRenderGridModule::Get().GetManager().CreateBatchRenderQueue(InThis, InJobs);
+		})
+	);
+}
+
 URenderGridQueue* URenderGrid::RenderSingleFrame(const int32 Frame)
 {
 	return RenderJobsSingleFrame(GetEnabledRenderGridJobs(), Frame);
@@ -1011,31 +1029,15 @@ URenderGridQueue* URenderGrid::RenderJobSingleFrame(URenderGridJob* Job, const i
 
 URenderGridQueue* URenderGrid::RenderJobsSingleFrame(const TArray<URenderGridJob*>& Jobs, const int32 Frame)
 {
-	if (!HasAnyFlags(RF_ClassDefaultObject))
-	{
-		if (URenderGrid* DefaultObject = Cast<URenderGrid>(GetClass()->GetDefaultObject(true)); IsValid(DefaultObject))
-		{
-			CopyAllProperties(DefaultObject);
-			CopyAllUserVariables(DefaultObject);
-			return DefaultObject->RenderJobsSingleFrame(Jobs, Frame);
-		}
-		return nullptr;
-	}
-
-	TArray<URenderGridJob*> JobsToRender;
-	for (URenderGridJob* Job : Jobs)
-	{
-		if (IsValid(Job) && HasRenderGridJob(Job))
-		{
-			JobsToRender.Add(Job);
-		}
-	}
-	if (URenderGridQueue* RenderQueue = UE::RenderGrid::IRenderGridModule::Get().GetManager().CreateBatchRenderQueueSingleFrame(this, Jobs, Frame); IsValid(RenderQueue))
-	{
-		RenderQueue->Execute();
-		return RenderQueue;
-	}
-	return nullptr;
+	return RunRenderJobs(
+		Jobs,
+		FRenderGridRunRenderJobsDefaultObjectCallback::CreateLambda([Frame](URenderGrid* InDefaultObject, const TArray<URenderGridJob*>& InJobs) -> URenderGridQueue* {
+			return InDefaultObject->RenderJobsSingleFrame(InJobs, Frame);
+		}),
+		FRenderGridRunRenderJobsCallback::CreateLambda([Frame](URenderGrid* InThis, const TArray<URenderGridJob*>& InJobs) -> URenderGridQueue* {
+			return UE::RenderGrid::IRenderGridModule::Get().GetManager().CreateBatchRenderQueueSingleFrame(InThis, InJobs, Frame);
+		})
+	);
 }
 
 URenderGridQueue* URenderGrid::RenderSingleFramePosition(const double FramePosition)
@@ -1050,31 +1052,20 @@ URenderGridQueue* URenderGrid::RenderJobSingleFramePosition(URenderGridJob* Job,
 
 URenderGridQueue* URenderGrid::RenderJobsSingleFramePosition(const TArray<URenderGridJob*>& Jobs, const double FramePosition)
 {
-	if (!HasAnyFlags(RF_ClassDefaultObject))
-	{
-		if (URenderGrid* DefaultObject = Cast<URenderGrid>(GetClass()->GetDefaultObject(true)); IsValid(DefaultObject))
-		{
-			CopyAllProperties(DefaultObject);
-			CopyAllUserVariables(DefaultObject);
-			return DefaultObject->RenderJobsSingleFramePosition(Jobs, FramePosition);
-		}
-		return nullptr;
-	}
+	return RunRenderJobs(
+		Jobs,
+		FRenderGridRunRenderJobsDefaultObjectCallback::CreateLambda([FramePosition](URenderGrid* InDefaultObject, const TArray<URenderGridJob*>& InJobs) -> URenderGridQueue* {
+			return InDefaultObject->RenderJobsSingleFramePosition(InJobs, FramePosition);
+		}),
+		FRenderGridRunRenderJobsCallback::CreateLambda([FramePosition](URenderGrid* InThis, const TArray<URenderGridJob*>& InJobs) -> URenderGridQueue* {
+			return UE::RenderGrid::IRenderGridModule::Get().GetManager().CreateBatchRenderQueueSingleFramePosition(InThis, InJobs, FramePosition);
+		})
+	);
+}
 
-	TArray<URenderGridJob*> JobsToRender;
-	for (URenderGridJob* Job : Jobs)
-	{
-		if (IsValid(Job) && HasRenderGridJob(Job))
-		{
-			JobsToRender.Add(Job);
-		}
-	}
-	if (URenderGridQueue* RenderQueue = UE::RenderGrid::IRenderGridModule::Get().GetManager().CreateBatchRenderQueueSingleFramePosition(this, Jobs, FramePosition); IsValid(RenderQueue))
-	{
-		RenderQueue->Execute();
-		return RenderQueue;
-	}
-	return nullptr;
+TSoftObjectPtr<UWorld> URenderGrid::GetMap() const
+{
+	return Settings->Map;
 }
 
 void URenderGrid::SetPropsSource(ERenderGridPropsSourceType InPropsSourceType, UObject* InPropsSourceOrigin)
diff --git a/Engine/Plugins/Experimental/RenderGrid/Source/RenderGrid/Private/RenderGrid/RenderGridManager.cpp b/Engine/Plugins/Experimental/RenderGrid/Source/RenderGrid/Private/RenderGrid/RenderGridManager.cpp
--- a/Engine/Plugins/Experimental/RenderGrid/Source/RenderGrid/Private/RenderGrid/RenderGridManager.cpp
+++ b/Engine/Plugins/Experimental/RenderGrid/Source/RenderGrid/Private/RenderGrid/RenderGridManager.cpp
@@ -58,7 +58,7 @@ URenderGridQueue* UE::RenderGrid::FRenderGridManager::RenderPreviewFrame(const F
 	const FRenderGridManagerRenderPreviewFrameArgsCallback Callback = Args.Callback;
 
 	URenderGrid* RenderGrid = Args.RenderGrid.Get();
-	if (!IsValid(RenderGrid))
+	if (!IsValid(RenderGrid) || RenderGrid->GetMap().IsNull() || !RenderGrid->GetMap().IsValid())
 	{
 		Callback.ExecuteIfBound(false);
 		return nullptr;
diff --git a/Engine/Plugins/Experimental/RenderGrid/Source/RenderGrid/Private/RenderGrid/RenderGridQueue.cpp b/Engine/Plugins/Experimental/RenderGrid/Source/RenderGrid/Private/RenderGrid/RenderGridQueue.cpp
--- a/Engine/Plugins/Experimental/RenderGrid/Source/RenderGrid/Private/RenderGrid/RenderGridQueue.cpp
+++ b/Engine/Plugins/Experimental/RenderGrid/Source/RenderGrid/Private/RenderGrid/RenderGridQueue.cpp
@@ -10,11 +10,11 @@
 #include "Utils/RenderGridGenericExecutionQueue.h"
 
 #include "Dom/JsonObject.h"
+#include "FileHelpers.h"
 #include "LevelSequence.h"
 #include "LevelSequenceEditorModule.h"
 #include "Modules/ModuleManager.h"
 #include "MoviePipelineAntiAliasingSetting.h"
-#include "MoviePipelineDeferredPasses.h"
 #include "MoviePipelineEditorBlueprintLibrary.h"
 #include "MoviePipelineExecutor.h"
 #include "MoviePipelineImageSequenceOutput.h"
@@ -76,6 +76,12 @@ URenderGridMoviePipelineRenderJob* URenderGridMoviePipelineRenderJob::Create(URe
 	UMoviePipelineExecutorJob* NewJob = UMoviePipelineEditorBlueprintLibrary::CreateJobFromSequence(RenderJob->PipelineQueue, JobSequence);
 	RenderJob->PipelineExecutorJob = NewJob;
 
+	FSoftObjectPath Map = RenderGrid->GetMap().ToSoftObjectPath();
+	if (Map.IsValid())
+	{
+		NewJob->Map = Map;
+	}
+
 	UMoviePipelinePrimaryConfig* JobRenderPreset = Job->GetRenderPreset();
 	if (IsValid(JobRenderPreset))
 	{
@@ -460,7 +466,7 @@ void URenderGridQueue::Tick(float DeltaTime)
 URenderGridQueue* URenderGridQueue::Create(const UE::RenderGrid::FRenderGridQueueCreateArgs& Args)
 {
 	URenderGrid* RenderGrid = Args.RenderGrid.Get();
-	if (!IsValid(RenderGrid))
+	if (!IsValid(RenderGrid) || RenderGrid->GetMap().IsNull())
 	{
 		return nullptr;
 	}
@@ -738,6 +744,15 @@ void URenderGridQueue::OnStart()
 	OnExecuteStartedDelegate.Clear();
 	AddToRoot();
 
+	Queue->Add(UE::RenderGrid::Private::FRenderGridGenericExecutionQueueAction::CreateLambda([this]()
+	{
+		if (!RenderGrid->GetMap().IsValid())
+		{
+			FEditorFileUtils::LoadMap(RenderGrid->GetMap().ToString(), false, true);
+		}
+	}));
+	Queue->DelayFrames(1);
+
 	Queue->Add(UE::RenderGrid::Private::FRenderGridGenericExecutionQueueAction::CreateLambda([this]()
 	{
 		PreviousFrameLimitSettings = UE::RenderGrid::Private::FRenderGridUtils::DisableFpsLimit();
diff --git a/Engine/Plugins/Experimental/RenderGrid/Source/RenderGrid/Public/RenderGrid/RenderGrid.h b/Engine/Plugins/Experimental/RenderGrid/Source/RenderGrid/Public/RenderGrid/RenderGrid.h
--- a/Engine/Plugins/Experimental/RenderGrid/Source/RenderGrid/Public/RenderGrid/RenderGrid.h
+++ b/Engine/Plugins/Experimental/RenderGrid/Source/RenderGrid/Public/RenderGrid/RenderGrid.h
@@ -54,11 +54,15 @@ public:
 
 	void CopyValuesFrom(URenderGridSettings* From)
 	{
+		Map = From->Map;
 		PropsSourceType = From->PropsSourceType;
 		PropsSourceOrigin_RemoteControl = From->PropsSourceOrigin_RemoteControl;
 	}
 
 public:
+	UPROPERTY(EditInstanceOnly, Category="Render Grid", DisplayName="Map", Meta=(AllowedClasses="/Script/Engine.World"))
+	TSoftObjectPtr<UWorld> Map;
+
 	/** The type of the properties that a job in this grid can have. */
 	UPROPERTY(/*EditInstanceOnly, Category="Render Grid", Meta=(DisplayName="Properties Type")*/)
 	ERenderGridPropsSourceType PropsSourceType;
@@ -559,6 +563,11 @@ public:
 	/** Overridable native event for when job rendering for the viewport viewer-mode ends. */
 	virtual void EndViewportRender(URenderGridJob* Job);
 
+private:
+	DECLARE_DELEGATE_RetVal_TwoParams(URenderGridQueue*, FRenderGridRunRenderJobsDefaultObjectCallback, URenderGrid* /*DefaultObject*/, TArray<URenderGridJob*> /*Jobs*/);
+	DECLARE_DELEGATE_RetVal_TwoParams(URenderGridQueue*, FRenderGridRunRenderJobsCallback, URenderGrid* /*This*/, TArray<URenderGridJob*> /*Jobs*/);
+	URenderGridQueue* RunRenderJobs(const TArray<URenderGridJob*>& Jobs, const FRenderGridRunRenderJobsDefaultObjectCallback& DefaultObjectCallback, const FRenderGridRunRenderJobsCallback& Callback);
+
 public:
 	/** Renders all the enabled jobs of this render grid. */
 	UFUNCTION(BlueprintCallable, Category="Render Grid")
@@ -610,6 +619,9 @@ public:
 public:
 	URenderGridSettings* GetSettingsObject() { return Settings; }
 
+	UFUNCTION(BlueprintPure, Category="Render Grid")
+	TSoftObjectPtr<UWorld> GetMap() const;
+
 	UFUNCTION(BlueprintCallable, Category="Render Grid")
 	void SetPropsSource(ERenderGridPropsSourceType InPropsSourceType, UObject* InPropsSourceOrigin = nullptr);
 
diff --git a/Engine/Plugins/Experimental/RenderGrid/Source/RenderGridDeveloper/Private/Blueprints/RenderGridBlueprint.cpp b/Engine/Plugins/Experimental/RenderGrid/Source/RenderGridDeveloper/Private/Blueprints/RenderGridBlueprint.cpp
--- a/Engine/Plugins/Experimental/RenderGrid/Source/RenderGridDeveloper/Private/Blueprints/RenderGridBlueprint.cpp
+++ b/Engine/Plugins/Experimental/RenderGrid/Source/RenderGridDeveloper/Private/Blueprints/RenderGridBlueprint.cpp
@@ -62,7 +62,7 @@ void URenderGridBlueprint::PostLoad()
 	OnPostVariablesChange(this);
 }
 
-void URenderGridBlueprint::RunOnInstances(const UE::RenderGrid::FRenderGridBlueprintRunOnInstancesCallback& Callback)
+void URenderGridBlueprint::RunOnInstances(const FRenderGridBlueprintRunOnInstancesCallback& Callback)
 {
 	if (UClass* MyRenderGridClass = GeneratedClass; IsValid(MyRenderGridClass))
 	{
@@ -107,7 +107,7 @@ void URenderGridBlueprint::Save()
 
 void URenderGridBlueprint::PropagateJobsToInstances()
 {
-	RunOnInstances(UE::RenderGrid::FRenderGridBlueprintRunOnInstancesCallback::CreateLambda([this](URenderGrid* Instance)
+	RunOnInstances(FRenderGridBlueprintRunOnInstancesCallback::CreateLambda([this](URenderGrid* Instance)
 	{
 		Instance->CopyJobs(RenderGrid);
 	}));
@@ -116,7 +116,7 @@ void URenderGridBlueprint::PropagateJobsToInstances()
 void URenderGridBlueprint::PropagateAllPropertiesExceptJobsToInstances()
 {
 	URenderGrid* CDO = GetRenderGridClassDefaultObject();
-	RunOnInstances(UE::RenderGrid::FRenderGridBlueprintRunOnInstancesCallback::CreateLambda([this, CDO](URenderGrid* Instance)
+	RunOnInstances(FRenderGridBlueprintRunOnInstancesCallback::CreateLambda([this, CDO](URenderGrid* Instance)
 	{
 		Instance->CopyAllPropertiesExceptJobs(RenderGrid);
 		if (IsValid(CDO))
@@ -129,7 +129,7 @@ void URenderGridBlueprint::PropagateAllPropertiesExceptJobsToInstances()
 void URenderGridBlueprint::PropagateAllPropertiesToInstances()
 {
 	URenderGrid* CDO = GetRenderGridClassDefaultObject();
-	RunOnInstances(UE::RenderGrid::FRenderGridBlueprintRunOnInstancesCallback::CreateLambda([this, CDO](URenderGrid* Instance)
+	RunOnInstances(FRenderGridBlueprintRunOnInstancesCallback::CreateLambda([this, CDO](URenderGrid* Instance)
 	{
 		Instance->CopyAllProperties(RenderGrid);
 		if (IsValid(CDO))
@@ -171,7 +171,7 @@ void URenderGridBlueprint::PropagateAllPropertiesToAsset(URenderGrid* Instance)
 URenderGrid* URenderGridBlueprint::GetRenderGridWithBlueprintGraph() const
 {
 	URenderGrid* Result = GetRenderGrid();
-	const_cast<URenderGridBlueprint*>(this)->RunOnInstances(UE::RenderGrid::FRenderGridBlueprintRunOnInstancesCallback::CreateLambda([&Result](URenderGrid* Instance)
+	const_cast<URenderGridBlueprint*>(this)->RunOnInstances(FRenderGridBlueprintRunOnInstancesCallback::CreateLambda([&Result](URenderGrid* Instance)
 	{
 		if (!IsValid(Result) || Result->HasAnyFlags(RF_ClassDefaultObject | RF_DefaultSubObject))
 		{
diff --git a/Engine/Plugins/Experimental/RenderGrid/Source/RenderGridDeveloper/Public/Blueprints/RenderGridBlueprint.h b/Engine/Plugins/Experimental/RenderGrid/Source/RenderGridDeveloper/Public/Blueprints/RenderGridBlueprint.h
--- a/Engine/Plugins/Experimental/RenderGrid/Source/RenderGridDeveloper/Public/Blueprints/RenderGridBlueprint.h
+++ b/Engine/Plugins/Experimental/RenderGrid/Source/RenderGridDeveloper/Public/Blueprints/RenderGridBlueprint.h
@@ -8,13 +8,6 @@
 #include "RenderGridBlueprint.generated.h"
 
 
-namespace UE::RenderGrid
-{
-	/** A delegate for URenderGridBlueprint::PropagateToInstances. */
-	DECLARE_DELEGATE_OneParam(FRenderGridBlueprintRunOnInstancesCallback, URenderGrid* /*Instance*/);
-}
-
-
 /**
  * A UBlueprint child class for the RenderGrid modules.
  *
@@ -45,7 +38,8 @@ public:
 	//~ End UBlueprint Interface
 
 private:
-	void RunOnInstances(const UE::RenderGrid::FRenderGridBlueprintRunOnInstancesCallback& Callback);
+	DECLARE_DELEGATE_OneParam(FRenderGridBlueprintRunOnInstancesCallback, URenderGrid* /*Instance*/);
+	void RunOnInstances(const FRenderGridBlueprintRunOnInstancesCallback& Callback);
 
 public:
 	void Load();
diff --git a/Engine/Plugins/Experimental/RenderGrid/Source/RenderGridEditor/Private/UI/SRenderGridViewerLive.cpp b/Engine/Plugins/Experimental/RenderGrid/Source/RenderGridEditor/Private/UI/SRenderGridViewerLive.cpp
--- a/Engine/Plugins/Experimental/RenderGrid/Source/RenderGridEditor/Private/UI/SRenderGridViewerLive.cpp
+++ b/Engine/Plugins/Experimental/RenderGrid/Source/RenderGridEditor/Private/UI/SRenderGridViewerLive.cpp
@@ -3,6 +3,7 @@
 #include "UI/SRenderGridViewerLive.h"
 #include "Camera/CameraComponent.h"
 #include "EngineUtils.h"
+#include "FileHelpers.h"
 #include "GameFramework/PlayerController.h"
 #include "GameFramework/PlayerStart.h"
 #include "IRenderGridEditor.h"
@@ -10,7 +11,9 @@
 #include "LevelSequenceActor.h"
 #include "RenderGrid/RenderGrid.h"
 #include "SlateOptMacros.h"
+#include "Styles/RenderGridEditorStyle.h"
 #include "UI/Components/SRenderGridViewerFrameSlider.h"
+#include "Widgets/Input/SButton.h"
 #include "Widgets/Layout/SScaleBox.h"
 
 #define LOCTEXT_NAMESPACE "SRenderGridViewerLive"
@@ -73,6 +76,11 @@ void UE::RenderGrid::Private::SRenderGridEditorViewport::Render()
 
 		if (URenderGrid* Grid = BlueprintEditor->GetInstance(); IsValid(Grid))
 		{
+			if (!Grid->GetMap().IsValid() || (Grid->GetMap().Get() != GetWorld()))
+			{
+				return;
+			}
+
 			if (ULevelSequencePlayer* SequencePlayer = GetSequencePlayer(); IsValid(SequencePlayer))
 			{
 				if (!Grid->HasRenderGridJob(RenderGridJob))
@@ -190,24 +198,33 @@ ULevelSequencePlayer* UE::RenderGrid::Private::SRenderGridEditorViewport::GetSeq
 
 void UE::RenderGrid::Private::SRenderGridEditorViewport::DestroySequencePlayer()
 {
-	TObjectPtr<ALevelSequenceActor> PlayerActor = LevelSequencePlayerActor;
-	TObjectPtr<ULevelSequencePlayer> Player = LevelSequencePlayer;
+	// TObjectPtr<ALevelSequenceActor> PlayerActor = LevelSequencePlayerActor;
+	// TObjectPtr<ULevelSequencePlayer> Player = LevelSequencePlayer;
 
 	LevelSequencePlayerWorld = nullptr;
 	LevelSequencePlayerActor = nullptr;
 	LevelSequencePlayer = nullptr;
 
-	if (IsValid(Player))
-	{
-		Player->Stop();
-	}
-	if (IsValid(PlayerActor))
+	//  Can cause crashes (when opening a different Map in the editor):
+	// if (IsValid(Player))
+	// {
+	//     Player->Stop();
+	// }
+	// if (IsValid(PlayerActor))
+	// {
+	//     PlayerActor->Destroy(false, false);
+	// }
+}
+
+
+void UE::RenderGrid::Private::SRenderGridViewerLive::Tick(const FGeometry&, const double, const float)
+{
+	if (const TSharedPtr<IRenderGridEditor> BlueprintEditor = BlueprintEditorWeakPtr.Pin())
 	{
-		PlayerActor->Destroy();
+		UpdateActionButton();
 	}
 }
 
-
 BEGIN_SLATE_FUNCTION_BUILD_OPTIMIZATION
 
 void UE::RenderGrid::Private::SRenderGridViewerLive::Construct(const FArguments& InArgs, TSharedPtr<IRenderGridEditor> InBlueprintEditor)
@@ -277,7 +294,7 @@ void UE::RenderGrid::Private::SRenderGridViewerLive::Construct(const FArguments&
 			.VAlign(VAlign_Fill)
 			[
 				SNew(SOverlay)
-				.Visibility_Lambda([this]() -> EVisibility { return (ViewportWidget->HasRenderedLastAttempt() ? EVisibility::Hidden : EVisibility::HitTestInvisible); })
+				.Visibility_Lambda([this]() -> EVisibility { return (ViewportWidget->HasRenderedLastAttempt() ? EVisibility::Hidden : EVisibility::Visible); })
 
 				// waiting text background (hides the viewport)
 				+ SOverlay::Slot()
@@ -289,34 +306,61 @@ void UE::RenderGrid::Private::SRenderGridViewerLive::Construct(const FArguments&
 					.ColorAndOpacity(FLinearColor(0, 0, 0, 1))
 				]
 
-				// waiting text
+				// waiting text button
 				+ SOverlay::Slot()
-				.HAlign(HAlign_Center)
-				.VAlign(VAlign_Center)
+				.HAlign(HAlign_Fill)
+				.VAlign(VAlign_Fill)
 				[
-					SNew(STextBlock)
-					.Text_Lambda([this]() -> FText
-					{
-						if (URenderGridJob* SelectedJob = SelectedJobWeakPtr.Get(); IsValid(SelectedJob))
-						{
-							if (ULevelSequence* Sequence = SelectedJob->GetLevelSequence(); IsValid(Sequence))
-							{
-								return LOCTEXT("WaitingForRenderer", "Waiting for renderer...");
-							}
-							return LOCTEXT("PleaseSelectLevelSequenceForJob", "Please select a level sequence for this job");
-						}
-						if (const TSharedPtr<IRenderGridEditor> BlueprintEditor = BlueprintEditorWeakPtr.Pin())
-						{
-							if (URenderGrid* Grid = BlueprintEditor->GetInstance(); IsValid(Grid))
+					SAssignNew(ActionButton, SButton)
+					.ContentPadding(0.0f)
+					.ButtonStyle(FRenderGridEditorStyle::Get(), TEXT("Invisible"))
+					.IsFocusable(false)
+					.OnClicked(this, &SRenderGridViewerLive::OnClicked)
+					[
+						SNew(SOverlay)
+
+						// waiting text
+						+ SOverlay::Slot()
+						.HAlign(HAlign_Center)
+						.VAlign(VAlign_Center)
+						[
+							SNew(STextBlock)
+							.Text_Lambda([this]() -> FText
 							{
-								if (!Grid->HasAnyRenderGridJobs())
+								if (const TSharedPtr<IRenderGridEditor> BlueprintEditor = BlueprintEditorWeakPtr.Pin())
 								{
-									return LOCTEXT("PleaseAddJob", "Please add a job");
+									if (URenderGrid* Grid = BlueprintEditor->GetInstance(); IsValid(Grid))
+									{
+										if (Grid->GetMap().IsNull())
+										{
+											return LOCTEXT("PleaseSelectMapForGrid", "Please select a map for this grid");
+										}
+										if (!Grid->GetMap().IsValid())
+										{
+											return LOCTEXT("ClickToOpenMapOfGrid", "Click here to open the map of this grid");
+										}
+										if (Grid->GetMap().Get() != ViewportWidget->GetWorldPublic())
+										{
+											return LOCTEXT("ViewerWrongWorldPleaseTryReopening", "The viewer is currently trying to view the incorrect world, please try reopening the RenderGrid");
+										}
+										if (!Grid->HasAnyRenderGridJobs())
+										{
+											return LOCTEXT("PleaseAddJob", "Please add a job");
+										}
+									}
 								}
-							}
-						}
-						return LOCTEXT("PleaseSelectJob", "Please select a job");
-					})
+								if (URenderGridJob* SelectedJob = SelectedJobWeakPtr.Get(); IsValid(SelectedJob))
+								{
+									if (ULevelSequence* Sequence = SelectedJob->GetLevelSequence(); IsValid(Sequence))
+									{
+										return LOCTEXT("WaitingForRenderer", "Waiting for renderer...");
+									}
+									return LOCTEXT("PleaseSelectLevelSequenceForJob", "Please select a level sequence for this job");
+								}
+								return LOCTEXT("PleaseSelectJob", "Please select a job");
+							})
+						]
+					]
 				]
 			]
 		]
@@ -333,6 +377,22 @@ void UE::RenderGrid::Private::SRenderGridViewerLive::Construct(const FArguments&
 
 END_SLATE_FUNCTION_BUILD_OPTIMIZATION
 
+FReply UE::RenderGrid::Private::SRenderGridViewerLive::OnClicked()
+{
+	if (const TSharedPtr<IRenderGridEditor> BlueprintEditor = BlueprintEditorWeakPtr.Pin())
+	{
+		if (URenderGrid* Grid = BlueprintEditor->GetInstance(); IsValid(Grid))
+		{
+			if (!Grid->GetMap().IsNull() && !Grid->GetMap().IsValid())
+			{
+				FEditorFileUtils::LoadMap(Grid->GetMap().ToString(), false, true);
+				return FReply::Handled();
+			}
+		}
+	}
+	return FReply::Handled();
+}
+
 void UE::RenderGrid::Private::SRenderGridViewerLive::OnObjectModified(UObject* Object)
 {
 	if (SelectedJobWeakPtr.IsValid() && (Object == SelectedJobWeakPtr))
@@ -404,6 +464,30 @@ void UE::RenderGrid::Private::SRenderGridViewerLive::UpdateViewport()
 	ViewportWidget->ClearSequenceFrame();
 }
 
+void UE::RenderGrid::Private::SRenderGridViewerLive::UpdateActionButton()
+{
+	if (!ActionButton.IsValid())
+	{
+		return;
+	}
+
+	bool bIsUsable = false;
+
+	if (const TSharedPtr<IRenderGridEditor> BlueprintEditor = BlueprintEditorWeakPtr.Pin())
+	{
+		if (URenderGrid* Grid = BlueprintEditor->GetInstance(); IsValid(Grid))
+		{
+			if (!Grid->GetMap().IsNull() && !Grid->GetMap().IsValid())
+			{
+				bIsUsable = true;
+			}
+		}
+	}
+
+	ActionButton->SetButtonStyle(&FRenderGridEditorStyle::Get().GetWidgetStyle<FButtonStyle>(bIsUsable ? TEXT("HoverHintOnly") : TEXT("Invisible")));
+	ActionButton->SetCursor(bIsUsable ? EMouseCursor::Type::Hand : EMouseCursor::Type::Default);
+}
+
 void UE::RenderGrid::Private::SRenderGridViewerLive::UpdateFrameSlider()
 {
 	if (!FrameSlider.IsValid())
diff --git a/Engine/Plugins/Experimental/RenderGrid/Source/RenderGridEditor/Private/UI/SRenderGridViewerLive.h b/Engine/Plugins/Experimental/RenderGrid/Source/RenderGridEditor/Private/UI/SRenderGridViewerLive.h
--- a/Engine/Plugins/Experimental/RenderGrid/Source/RenderGridEditor/Private/UI/SRenderGridViewerLive.h
+++ b/Engine/Plugins/Experimental/RenderGrid/Source/RenderGridEditor/Private/UI/SRenderGridViewerLive.h
@@ -10,6 +10,7 @@
 
 
 class ALevelSequenceActor;
+class SButton;
 class SSlider;
 class URenderGridJob;
 class URenderGrid;
@@ -59,6 +60,7 @@ namespace UE::RenderGrid::Private
 		void ClearSequenceFrame();
 		bool ShowSequenceFrame(URenderGridJob* InJob, ULevelSequence* InSequence, const float InTime);
 		bool HasRenderedLastAttempt() const { return bRenderedLastAttempt; }
+		UWorld* GetWorldPublic() const { return GetWorld(); }
 
 	protected:
 		//~ Begin SEditorViewport Interface
@@ -116,14 +118,18 @@ namespace UE::RenderGrid::Private
 		SLATE_BEGIN_ARGS(SRenderGridViewerLive) {}
 		SLATE_END_ARGS()
 
+		virtual void Tick(const FGeometry&, const double, const float) override;
 		void Construct(const FArguments& InArgs, TSharedPtr<IRenderGridEditor> InBlueprintEditor);
 
 	private:
+		FReply OnClicked();
 		void OnObjectModified(UObject* Object);
 		void RenderGridJobDataChanged();
 		void SelectedJobChanged();
 		void FrameSliderValueChanged(const float NewValue);
+
 		void UpdateViewport();
+		void UpdateActionButton();
 		void UpdateFrameSlider();
 
 	private:
@@ -138,5 +144,8 @@ namespace UE::RenderGrid::Private
 
 		/** The widget that allows the user to select what frame they'd like to see. */
 		TSharedPtr<SRenderGridViewerFrameSlider> FrameSlider;
+
+		/** The button which can be clicked, to to rerender the preview, or to open the map of the grid. */
+		TSharedPtr<SButton> ActionButton;
 	};
 }
diff --git a/Engine/Plugins/Experimental/RenderGrid/Source/RenderGridEditor/Private/UI/SRenderGridViewerPreview.cpp b/Engine/Plugins/Experimental/RenderGrid/Source/RenderGridEditor/Private/UI/SRenderGridViewerPreview.cpp
--- a/Engine/Plugins/Experimental/RenderGrid/Source/RenderGridEditor/Private/UI/SRenderGridViewerPreview.cpp
+++ b/Engine/Plugins/Experimental/RenderGrid/Source/RenderGridEditor/Private/UI/SRenderGridViewerPreview.cpp
@@ -19,6 +19,7 @@
 #include "Widgets/Colors/SColorPicker.h"
 #include "Widgets/Input/SButton.h"
 #include "Widgets/Layout/SScaleBox.h"
+#include "FileHelpers.h"
 
 #define LOCTEXT_NAMESPACE "SRenderGridViewerPreview"
 
@@ -30,9 +31,10 @@ void UE::RenderGrid::Private::SRenderGridViewerPreview::Tick(const FGeometry&, c
 {
 	if (const TSharedPtr<IRenderGridEditor> BlueprintEditor = BlueprintEditorWeakPtr.Pin())
 	{
+		UpdateActionButton();
+
 		if (BlueprintEditor->CanCurrentlyRender() && !CurrentRenderQueueWeakPtr.IsValid())
 		{
-			UpdateRerenderButton();
 			UpdateFrameSlider();
 
 			if (FramesUntilRenderNewPreview > 0)
@@ -105,7 +107,7 @@ void UE::RenderGrid::Private::SRenderGridViewerPreview::Construct(const FArgumen
 			.HAlign(HAlign_Fill)
 			.VAlign(VAlign_Fill)
 			[
-				SAssignNew(RerenderButton, SButton)
+				SAssignNew(ActionButton, SButton)
 				.ContentPadding(0.0f)
 				.ButtonStyle(FRenderGridEditorStyle::Get(), TEXT("Invisible"))
 				.IsFocusable(false)
@@ -121,27 +123,35 @@ void UE::RenderGrid::Private::SRenderGridViewerPreview::Construct(const FArgumen
 						SNew(STextBlock)
 						.Text_Lambda([this]() -> FText
 						{
-							if (URenderGridJob* SelectedJob = SelectedJobWeakPtr.Get(); IsValid(SelectedJob))
+							if (const TSharedPtr<IRenderGridEditor> BlueprintEditor = BlueprintEditorWeakPtr.Pin())
 							{
-								if (ULevelSequence* Sequence = SelectedJob->GetLevelSequence(); IsValid(Sequence))
+								if (URenderGrid* Grid = BlueprintEditor->GetInstance(); IsValid(Grid))
 								{
-									if (IsPreviewWidget() || !FrameSlider.IsValid() || (FrameSlider->GetVisibility() != EVisibility::Visible))
+									if (Grid->GetMap().IsNull())
 									{
-										return LOCTEXT("WaitingForRenderer", "Waiting for renderer...");
+										return LOCTEXT("PleaseSelectMapForGrid", "Please select a map for this grid");
+									}
+									if (!Grid->GetMap().IsValid())
+									{
+										return LOCTEXT("ClickToOpenMapOfGrid", "Click here to open the map of this grid");
+									}
+									if (!Grid->HasAnyRenderGridJobs())
+									{
+										return LOCTEXT("PleaseAddJob", "Please add a job");
 									}
-									return LOCTEXT("ClickToStartRendering", "Click here to start rendering");
 								}
-								return LOCTEXT("PleaseSelectLevelSequenceForJob", "Please select a level sequence for this job");
 							}
-							if (const TSharedPtr<IRenderGridEditor> BlueprintEditor = BlueprintEditorWeakPtr.Pin())
+							if (URenderGridJob* SelectedJob = SelectedJobWeakPtr.Get(); IsValid(SelectedJob))
 							{
-								if (URenderGrid* Grid = BlueprintEditor->GetInstance(); IsValid(Grid))
+								if (ULevelSequence* Sequence = SelectedJob->GetLevelSequence(); IsValid(Sequence))
 								{
-									if (!Grid->HasAnyRenderGridJobs())
+									if (IsPreviewWidget() || !FrameSlider.IsValid() || (FrameSlider->GetVisibility() != EVisibility::Visible))
 									{
-										return LOCTEXT("PleaseAddJob", "Please add a job");
+										return LOCTEXT("WaitingForRenderer", "Waiting for renderer...");
 									}
+									return LOCTEXT("ClickToStartRendering", "Click here to start rendering");
 								}
+								return LOCTEXT("PleaseSelectLevelSequenceForJob", "Please select a level sequence for this job");
 							}
 							return LOCTEXT("PleaseSelectJob", "Please select a job");
 						})
@@ -196,11 +206,22 @@ END_SLATE_FUNCTION_BUILD_OPTIMIZATION
 
 FReply UE::RenderGrid::Private::SRenderGridViewerPreview::OnClicked()
 {
+	if (const TSharedPtr<IRenderGridEditor> BlueprintEditor = BlueprintEditorWeakPtr.Pin())
+	{
+		if (URenderGrid* Grid = BlueprintEditor->GetInstance(); IsValid(Grid))
+		{
+			if (!Grid->GetMap().IsNull() && !Grid->GetMap().IsValid())
+			{
+				FEditorFileUtils::LoadMap(Grid->GetMap().ToString(), false, true);
+				return FReply::Handled();
+			}
+		}
+	}
+
 	if (!IsPreviewWidget())
 	{
 		RenderNewPreview();
 	}
-
 	return FReply::Handled();
 }
 
@@ -322,6 +343,17 @@ void UE::RenderGrid::Private::SRenderGridViewerPreview::InternalRenderNewPreview
 
 	if (const TSharedPtr<IRenderGridEditor> BlueprintEditor = BlueprintEditorWeakPtr.Pin())
 	{
+		if (URenderGrid* Grid = BlueprintEditor->GetInstance(); IsValid(Grid))
+		{
+			if (Grid->GetMap().IsNull() || !Grid->GetMap().IsValid())
+			{
+				// don't render, clear image and try again next frame
+				SetImageTexture(nullptr);
+				RenderNewPreview();
+				return;
+			}
+		}
+
 		double WidgetWidth = FMath::Max(120.0, FMath::Min(GetTickSpaceGeometry().Size.X, GetTickSpaceGeometry().Size.Y * Job->GetOutputAspectRatio()));
 		double RenderResolution = WidgetWidth * (IsPreviewWidget() ? 1.25 : 0.75); // pixels in width
 
@@ -482,16 +514,28 @@ void UE::RenderGrid::Private::SRenderGridViewerPreview::SetImageTexture(UTexture
 }
 
 
-void UE::RenderGrid::Private::SRenderGridViewerPreview::UpdateRerenderButton()
+void UE::RenderGrid::Private::SRenderGridViewerPreview::UpdateActionButton()
 {
-	if (!RerenderButton.IsValid())
+	if (!ActionButton.IsValid())
 	{
 		return;
 	}
-	const bool bIsUsable = (!IsPreviewWidget() && SelectedJobWeakPtr.IsValid());
 
-	RerenderButton->SetButtonStyle(&FRenderGridEditorStyle::Get().GetWidgetStyle<FButtonStyle>(bIsUsable ? TEXT("HoverHintOnly") : TEXT("Invisible")));
-	RerenderButton->SetCursor(bIsUsable ? EMouseCursor::Type::Hand : EMouseCursor::Type::Default);
+	bool bIsUsable = (!IsPreviewWidget() && SelectedJobWeakPtr.IsValid());
+
+	if (const TSharedPtr<IRenderGridEditor> BlueprintEditor = BlueprintEditorWeakPtr.Pin())
+	{
+		if (URenderGrid* Grid = BlueprintEditor->GetInstance(); IsValid(Grid))
+		{
+			if (!Grid->GetMap().IsNull() && !Grid->GetMap().IsValid())
+			{
+				bIsUsable = true;
+			}
+		}
+	}
+
+	ActionButton->SetButtonStyle(&FRenderGridEditorStyle::Get().GetWidgetStyle<FButtonStyle>(bIsUsable ? TEXT("HoverHintOnly") : TEXT("Invisible")));
+	ActionButton->SetCursor(bIsUsable ? EMouseCursor::Type::Hand : EMouseCursor::Type::Default);
 }
 
 void UE::RenderGrid::Private::SRenderGridViewerPreview::UpdateFrameSlider()
diff --git a/Engine/Plugins/Experimental/RenderGrid/Source/RenderGridEditor/Private/UI/SRenderGridViewerPreview.h b/Engine/Plugins/Experimental/RenderGrid/Source/RenderGridEditor/Private/UI/SRenderGridViewerPreview.h
--- a/Engine/Plugins/Experimental/RenderGrid/Source/RenderGridEditor/Private/UI/SRenderGridViewerPreview.h
+++ b/Engine/Plugins/Experimental/RenderGrid/Source/RenderGridEditor/Private/UI/SRenderGridViewerPreview.h
@@ -67,10 +67,9 @@ namespace UE::RenderGrid::Private
 		void InternalRenderNewPreviewOfRenderGridJob(URenderGridJob* Job);
 		void RenderNewPreviewCallback(const bool bSuccess);
 
-		void UpdateRerenderButton();
+		void UpdateActionButton();
 		void UpdateFrameSlider();
 
-
 	private:
 		/** A reference to the blueprint editor that owns the render grid instance. */
 		TWeakPtr<IRenderGridEditor> BlueprintEditorWeakPtr;
@@ -91,8 +90,8 @@ namespace UE::RenderGrid::Private
 		/** The widget that allows the user to select what frame they'd like to see. */
 		TSharedPtr<SRenderGridViewerFrameSlider> FrameSlider;
 
-		/** The button which can be clicked to rerender the preview. */
-		TSharedPtr<SButton> RerenderButton;
+		/** The button which can be clicked, to to rerender the preview, or to open the map of the grid. */
+		TSharedPtr<SButton> ActionButton;
 
 		/** The widget that contains the image. */
 		TSharedPtr<SImage> Image;
