RenderGrid plugin:
- Renamed the "Map" variable to "Level".

diff --git a/Engine/Plugins/Experimental/RenderGrid/Source/RenderGrid/Private/RenderGrid/RenderGrid.cpp b/Engine/Plugins/Experimental/RenderGrid/Source/RenderGrid/Private/RenderGrid/RenderGrid.cpp
--- a/Engine/Plugins/Experimental/RenderGrid/Source/RenderGrid/Private/RenderGrid/RenderGrid.cpp
+++ b/Engine/Plugins/Experimental/RenderGrid/Source/RenderGrid/Private/RenderGrid/RenderGrid.cpp
@@ -1063,9 +1063,9 @@ URenderGridQueue* URenderGrid::RenderJobsSingleFramePosition(const TArray<URende
 	);
 }
 
-TSoftObjectPtr<UWorld> URenderGrid::GetMap() const
+TSoftObjectPtr<UWorld> URenderGrid::GetLevel() const
 {
-	return Settings->Map;
+	return Settings->Level;
 }
 
 void URenderGrid::SetPropsSource(ERenderGridPropsSourceType InPropsSourceType, UObject* InPropsSourceOrigin)
diff --git a/Engine/Plugins/Experimental/RenderGrid/Source/RenderGrid/Private/RenderGrid/RenderGridManager.cpp b/Engine/Plugins/Experimental/RenderGrid/Source/RenderGrid/Private/RenderGrid/RenderGridManager.cpp
--- a/Engine/Plugins/Experimental/RenderGrid/Source/RenderGrid/Private/RenderGrid/RenderGridManager.cpp
+++ b/Engine/Plugins/Experimental/RenderGrid/Source/RenderGrid/Private/RenderGrid/RenderGridManager.cpp
@@ -58,7 +58,7 @@ URenderGridQueue* UE::RenderGrid::FRenderGridManager::RenderPreviewFrame(const F
 	const FRenderGridManagerRenderPreviewFrameArgsCallback Callback = Args.Callback;
 
 	URenderGrid* RenderGrid = Args.RenderGrid.Get();
-	if (!IsValid(RenderGrid) || RenderGrid->GetMap().IsNull() || !RenderGrid->GetMap().IsValid())
+	if (!IsValid(RenderGrid) || RenderGrid->GetLevel().IsNull() || !RenderGrid->GetLevel().IsValid())
 	{
 		Callback.ExecuteIfBound(false);
 		return nullptr;
diff --git a/Engine/Plugins/Experimental/RenderGrid/Source/RenderGrid/Private/RenderGrid/RenderGridQueue.cpp b/Engine/Plugins/Experimental/RenderGrid/Source/RenderGrid/Private/RenderGrid/RenderGridQueue.cpp
--- a/Engine/Plugins/Experimental/RenderGrid/Source/RenderGrid/Private/RenderGrid/RenderGridQueue.cpp
+++ b/Engine/Plugins/Experimental/RenderGrid/Source/RenderGrid/Private/RenderGrid/RenderGridQueue.cpp
@@ -76,10 +76,10 @@ URenderGridMoviePipelineRenderJob* URenderGridMoviePipelineRenderJob::Create(URe
 	UMoviePipelineExecutorJob* NewJob = UMoviePipelineEditorBlueprintLibrary::CreateJobFromSequence(RenderJob->PipelineQueue, JobSequence);
 	RenderJob->PipelineExecutorJob = NewJob;
 
-	FSoftObjectPath Map = RenderGrid->GetMap().ToSoftObjectPath();
-	if (Map.IsValid())
+	FSoftObjectPath Level = RenderGrid->GetLevel().ToSoftObjectPath();
+	if (Level.IsValid())
 	{
-		NewJob->Map = Map;
+		NewJob->Map = Level;
 	}
 
 	UMoviePipelinePrimaryConfig* JobRenderPreset = Job->GetRenderPreset();
@@ -466,7 +466,7 @@ void URenderGridQueue::Tick(float DeltaTime)
 URenderGridQueue* URenderGridQueue::Create(const UE::RenderGrid::FRenderGridQueueCreateArgs& Args)
 {
 	URenderGrid* RenderGrid = Args.RenderGrid.Get();
-	if (!IsValid(RenderGrid) || RenderGrid->GetMap().IsNull())
+	if (!IsValid(RenderGrid) || RenderGrid->GetLevel().IsNull())
 	{
 		return nullptr;
 	}
@@ -746,9 +746,9 @@ void URenderGridQueue::OnStart()
 
 	Queue->Add(UE::RenderGrid::Private::FRenderGridGenericExecutionQueueAction::CreateLambda([this]()
 	{
-		if (!RenderGrid->GetMap().IsValid())
+		if (!RenderGrid->GetLevel().IsValid())
 		{
-			FEditorFileUtils::LoadMap(RenderGrid->GetMap().ToString(), false, true);
+			FEditorFileUtils::LoadMap(RenderGrid->GetLevel().ToString(), false, true);
 		}
 	}));
 	Queue->DelayFrames(1);
diff --git a/Engine/Plugins/Experimental/RenderGrid/Source/RenderGrid/Public/RenderGrid/RenderGrid.h b/Engine/Plugins/Experimental/RenderGrid/Source/RenderGrid/Public/RenderGrid/RenderGrid.h
--- a/Engine/Plugins/Experimental/RenderGrid/Source/RenderGrid/Public/RenderGrid/RenderGrid.h
+++ b/Engine/Plugins/Experimental/RenderGrid/Source/RenderGrid/Public/RenderGrid/RenderGrid.h
@@ -54,14 +54,14 @@ public:
 
 	void CopyValuesFrom(URenderGridSettings* From)
 	{
-		Map = From->Map;
+		Level = From->Level;
 		PropsSourceType = From->PropsSourceType;
 		PropsSourceOrigin_RemoteControl = From->PropsSourceOrigin_RemoteControl;
 	}
 
 public:
-	UPROPERTY(EditInstanceOnly, Category="Render Grid", DisplayName="Map", Meta=(AllowedClasses="/Script/Engine.World"))
-	TSoftObjectPtr<UWorld> Map;
+	UPROPERTY(EditInstanceOnly, Category="Render Grid", DisplayName="Level", Meta=(AllowedClasses="/Script/Engine.World"))
+	TSoftObjectPtr<UWorld> Level;
 
 	/** The type of the properties that a job in this grid can have. */
 	UPROPERTY(/*EditInstanceOnly, Category="Render Grid", Meta=(DisplayName="Properties Type")*/)
@@ -620,7 +620,7 @@ public:
 	URenderGridSettings* GetSettingsObject() { return Settings; }
 
 	UFUNCTION(BlueprintPure, Category="Render Grid")
-	TSoftObjectPtr<UWorld> GetMap() const;
+	TSoftObjectPtr<UWorld> GetLevel() const;
 
 	UFUNCTION(BlueprintCallable, Category="Render Grid")
 	void SetPropsSource(ERenderGridPropsSourceType InPropsSourceType, UObject* InPropsSourceOrigin = nullptr);
diff --git a/Engine/Plugins/Experimental/RenderGrid/Source/RenderGridEditor/Private/UI/SRenderGridViewerLive.cpp b/Engine/Plugins/Experimental/RenderGrid/Source/RenderGridEditor/Private/UI/SRenderGridViewerLive.cpp
--- a/Engine/Plugins/Experimental/RenderGrid/Source/RenderGridEditor/Private/UI/SRenderGridViewerLive.cpp
+++ b/Engine/Plugins/Experimental/RenderGrid/Source/RenderGridEditor/Private/UI/SRenderGridViewerLive.cpp
@@ -76,7 +76,7 @@ void UE::RenderGrid::Private::SRenderGridEditorViewport::Render()
 
 		if (URenderGrid* Grid = BlueprintEditor->GetInstance(); IsValid(Grid))
 		{
-			if (!Grid->GetMap().IsValid() || (Grid->GetMap().Get() != GetWorld()))
+			if (!Grid->GetLevel().IsValid() || (Grid->GetLevel().Get() != GetWorld()))
 			{
 				return;
 			}
@@ -331,15 +331,15 @@ void UE::RenderGrid::Private::SRenderGridViewerLive::Construct(const FArguments&
 								{
 									if (URenderGrid* Grid = BlueprintEditor->GetInstance(); IsValid(Grid))
 									{
-										if (Grid->GetMap().IsNull())
+										if (Grid->GetLevel().IsNull())
 										{
-											return LOCTEXT("PleaseSelectMapForGrid", "Please select a map for this grid");
+											return LOCTEXT("PleaseSelectLevelForGrid", "Please select a level for this grid");
 										}
-										if (!Grid->GetMap().IsValid())
+										if (!Grid->GetLevel().IsValid())
 										{
-											return LOCTEXT("ClickToOpenMapOfGrid", "Click here to open the map of this grid");
+											return LOCTEXT("ClickToOpenLevelOfGrid", "Click here to open the level of this grid");
 										}
-										if (Grid->GetMap().Get() != ViewportWidget->GetWorldPublic())
+										if (Grid->GetLevel().Get() != ViewportWidget->GetWorldPublic())
 										{
 											return LOCTEXT("ViewerWrongWorldPleaseTryReopening", "The viewer is currently trying to view the incorrect world, please try reopening the RenderGrid");
 										}
@@ -383,9 +383,9 @@ FReply UE::RenderGrid::Private::SRenderGridViewerLive::OnClicked()
 	{
 		if (URenderGrid* Grid = BlueprintEditor->GetInstance(); IsValid(Grid))
 		{
-			if (!Grid->GetMap().IsNull() && !Grid->GetMap().IsValid())
+			if (!Grid->GetLevel().IsNull() && !Grid->GetLevel().IsValid())
 			{
-				FEditorFileUtils::LoadMap(Grid->GetMap().ToString(), false, true);
+				FEditorFileUtils::LoadMap(Grid->GetLevel().ToString(), false, true);
 				return FReply::Handled();
 			}
 		}
@@ -477,7 +477,7 @@ void UE::RenderGrid::Private::SRenderGridViewerLive::UpdateActionButton()
 	{
 		if (URenderGrid* Grid = BlueprintEditor->GetInstance(); IsValid(Grid))
 		{
-			if (!Grid->GetMap().IsNull() && !Grid->GetMap().IsValid())
+			if (!Grid->GetLevel().IsNull() && !Grid->GetLevel().IsValid())
 			{
 				bIsUsable = true;
 			}
diff --git a/Engine/Plugins/Experimental/RenderGrid/Source/RenderGridEditor/Private/UI/SRenderGridViewerPreview.cpp b/Engine/Plugins/Experimental/RenderGrid/Source/RenderGridEditor/Private/UI/SRenderGridViewerPreview.cpp
--- a/Engine/Plugins/Experimental/RenderGrid/Source/RenderGridEditor/Private/UI/SRenderGridViewerPreview.cpp
+++ b/Engine/Plugins/Experimental/RenderGrid/Source/RenderGridEditor/Private/UI/SRenderGridViewerPreview.cpp
@@ -127,13 +127,13 @@ void UE::RenderGrid::Private::SRenderGridViewerPreview::Construct(const FArgumen
 							{
 								if (URenderGrid* Grid = BlueprintEditor->GetInstance(); IsValid(Grid))
 								{
-									if (Grid->GetMap().IsNull())
+									if (Grid->GetLevel().IsNull())
 									{
-										return LOCTEXT("PleaseSelectMapForGrid", "Please select a map for this grid");
+										return LOCTEXT("PleaseSelectLevelForGrid", "Please select a level for this grid");
 									}
-									if (!Grid->GetMap().IsValid())
+									if (!Grid->GetLevel().IsValid())
 									{
-										return LOCTEXT("ClickToOpenMapOfGrid", "Click here to open the map of this grid");
+										return LOCTEXT("ClickToOpenLevelOfGrid", "Click here to open the level of this grid");
 									}
 									if (!Grid->HasAnyRenderGridJobs())
 									{
@@ -210,9 +210,9 @@ FReply UE::RenderGrid::Private::SRenderGridViewerPreview::OnClicked()
 	{
 		if (URenderGrid* Grid = BlueprintEditor->GetInstance(); IsValid(Grid))
 		{
-			if (!Grid->GetMap().IsNull() && !Grid->GetMap().IsValid())
+			if (!Grid->GetLevel().IsNull() && !Grid->GetLevel().IsValid())
 			{
-				FEditorFileUtils::LoadMap(Grid->GetMap().ToString(), false, true);
+				FEditorFileUtils::LoadMap(Grid->GetLevel().ToString(), false, true);
 				return FReply::Handled();
 			}
 		}
@@ -345,7 +345,7 @@ void UE::RenderGrid::Private::SRenderGridViewerPreview::InternalRenderNewPreview
 	{
 		if (URenderGrid* Grid = BlueprintEditor->GetInstance(); IsValid(Grid))
 		{
-			if (Grid->GetMap().IsNull() || !Grid->GetMap().IsValid())
+			if (Grid->GetLevel().IsNull() || !Grid->GetLevel().IsValid())
 			{
 				// don't render, clear image and try again next frame
 				SetImageTexture(nullptr);
@@ -527,7 +527,7 @@ void UE::RenderGrid::Private::SRenderGridViewerPreview::UpdateActionButton()
 	{
 		if (URenderGrid* Grid = BlueprintEditor->GetInstance(); IsValid(Grid))
 		{
-			if (!Grid->GetMap().IsNull() && !Grid->GetMap().IsValid())
+			if (!Grid->GetLevel().IsNull() && !Grid->GetLevel().IsValid())
 			{
 				bIsUsable = true;
 			}
