RenderGrid plugin:
- Removed the job filter button (as there are currently no filter options).
- Made new render grid assets be automatically compiled upon first launch.

diff --git a/Engine/Plugins/Experimental/RenderGrid/Source/RenderGridDeveloper/Private/Blueprints/RenderGridBlueprint.cpp b/Engine/Plugins/Experimental/RenderGrid/Source/RenderGridDeveloper/Private/Blueprints/RenderGridBlueprint.cpp
--- a/Engine/Plugins/Experimental/RenderGrid/Source/RenderGridDeveloper/Private/Blueprints/RenderGridBlueprint.cpp
+++ b/Engine/Plugins/Experimental/RenderGrid/Source/RenderGridDeveloper/Private/Blueprints/RenderGridBlueprint.cpp
@@ -54,6 +54,8 @@ void URenderGridBlueprint::PostLoad()
 		FBlueprintEditorUtils::MarkBlueprintAsStructurallyModified(this);
 		FBlueprintEditorUtils::AddUbergraphPage(this, NewGraph);
 		LastEditedDocuments.AddUnique(NewGraph);
+
+		FKismetEditorUtilities::CompileBlueprint(this, EBlueprintCompileOptions::SkipGarbageCollection);
 	}
 
 	OnChanged().RemoveAll(this);
diff --git a/Engine/Plugins/Experimental/RenderGrid/Source/RenderGridEditor/Private/UI/SRenderGridJobList.cpp b/Engine/Plugins/Experimental/RenderGrid/Source/RenderGridEditor/Private/UI/SRenderGridJobList.cpp
--- a/Engine/Plugins/Experimental/RenderGrid/Source/RenderGridEditor/Private/UI/SRenderGridJobList.cpp
+++ b/Engine/Plugins/Experimental/RenderGrid/Source/RenderGridEditor/Private/UI/SRenderGridJobList.cpp
@@ -13,7 +13,6 @@
 
 #include "Framework/MultiBox/MultiBoxBuilder.h"
 #include "Misc/MessageDialog.h"
-#include "MoviePipelineExecutor.h"
 #include "MoviePipelineOutputSetting.h"
 #include "MoviePipelinePrimaryConfig.h"
 #include "MoviePipelineQueue.h"
@@ -26,7 +25,6 @@
 #include "Styling/AppStyle.h"
 #include "Tracks/MovieSceneSubTrack.h"
 #include "Widgets/Input/SCheckBox.h"
-#include "Widgets/Input/SComboButton.h"
 #include "Widgets/Input/SSearchBox.h"
 #include "Widgets/Layout/SScaleBox.h"
 
@@ -95,7 +93,7 @@ void UE::RenderGrid::Private::SRenderGridJobList::Construct(const FArguments& In
 				.OnTextChanged(this, &SRenderGridJobList::OnSearchBarTextChanged)
 			]
 
-			// Filters
+			/* // Filters
 			+ SHorizontalBox::Slot()
 			.AutoWidth()
 			.Padding(0.f, 2.f, 2.f, 2.f)
@@ -111,9 +109,9 @@ void UE::RenderGrid::Private::SRenderGridJobList::Construct(const FArguments& In
 					SNew(STextBlock)
 					.TextStyle(FAppStyle::Get(), "GenericFilters.TextStyle")
 					.Font(FAppStyle::Get().GetFontStyle("FontAwesome.12"))
-					.Text(FText::FromString(FString(TEXT("\xf0b0"))) /*fa-filter*/)
+					.Text(FText::FromString(FString(TEXT("\xf0b0")))) // fa-filter
 				]
-			]
+			] */
 		]
 
 		// Job List
