Fixing compile errors

diff --git a/Engine/Plugins/Experimental/RenderGrid/Source/RenderGrid/Private/RenderGrid/RenderGridQueue.cpp b/Engine/Plugins/Experimental/RenderGrid/Source/RenderGrid/Private/RenderGrid/RenderGridQueue.cpp
--- a/Engine/Plugins/Experimental/RenderGrid/Source/RenderGrid/Private/RenderGrid/RenderGridQueue.cpp
+++ b/Engine/Plugins/Experimental/RenderGrid/Source/RenderGrid/Private/RenderGrid/RenderGridQueue.cpp
@@ -626,7 +626,7 @@ TArray<URenderGridJob*> URenderGridQueue::GetJobs() const
 {
 	TArray<URenderGridJob*> Jobs;
 	Jobs.Reserve(Entries.Num());
-	for (const TTuple<TObjectPtr<URenderGridJob>, TObjectPtr<URenderGridMoviePipelineRenderJob>> EntryEntry : Entries)
+	for (const TTuple<TObjectPtr<URenderGridJob>, TObjectPtr<URenderGridMoviePipelineRenderJob>>& EntryEntry : Entries)
 	{
 		if (const TObjectPtr<URenderGridJob> Job = EntryEntry.Key; IsValid(Job))
 		{
@@ -645,7 +645,7 @@ TArray<URenderGridJob*> URenderGridQueue::GetJobs() const
 int32 URenderGridQueue::GetJobsCount() const
 {
 	int32 JobsCount = 0;
-	for (const TTuple<TObjectPtr<URenderGridJob>, TObjectPtr<URenderGridMoviePipelineRenderJob>> EntryEntry : Entries)
+	for (const TTuple<TObjectPtr<URenderGridJob>, TObjectPtr<URenderGridMoviePipelineRenderJob>>& EntryEntry : Entries)
 	{
 		if (const TObjectPtr<URenderGridJob> Job = EntryEntry.Key; IsValid(Job))
 		{
@@ -669,7 +669,7 @@ int32 URenderGridQueue::GetJobsRemainingCount() const
 	}
 
 	int32 JobsCount = 0;
-	for (const TObjectPtr<URenderGridJob> Job : RemainingJobs)
+	for (const TObjectPtr<URenderGridJob>& Job : RemainingJobs)
 	{
 		if (const TObjectPtr<URenderGridMoviePipelineRenderJob>* EntryPtr = Entries.Find(Job))
 		{
diff --git a/Engine/Plugins/Runtime/OpenXR/Source/OpenXRInput/Private/OpenXRInput.cpp b/Engine/Plugins/Runtime/OpenXR/Source/OpenXRInput/Private/OpenXRInput.cpp
--- a/Engine/Plugins/Runtime/OpenXR/Source/OpenXRInput/Private/OpenXRInput.cpp
+++ b/Engine/Plugins/Runtime/OpenXR/Source/OpenXRInput/Private/OpenXRInput.cpp
@@ -555,7 +555,7 @@ void FOpenXRInputPlugin::FOpenXRInput::BuildEnhancedActions(TMap<FString, FInter
 	}
 
 	PRAGMA_DISABLE_DEPRECATION_WARNINGS
-	for (const TPair<TObjectPtr<UInputMappingContext>, int32> MappingContext : MappableInputConfig->GetMappingContexts())
+	for (const TPair<TObjectPtr<UInputMappingContext>, int32>& MappingContext : MappableInputConfig->GetMappingContexts())
 	{
 		FOpenXRActionSet ActionSet(Instance, MappingContext.Key->GetFName(), MappingContext.Key->ContextDescription.ToString(), ToXrPriority(MappingContext.Value), MappingContext.Key);
 		TMap<FName, int32> ActionMap;
diff --git a/Engine/Plugins/VirtualProduction/LevelSnapshots/Source/LevelSnapshotsEditor/Private/Widgets/Filter/SFilterSearchMenu.cpp b/Engine/Plugins/VirtualProduction/LevelSnapshots/Source/LevelSnapshotsEditor/Private/Widgets/Filter/SFilterSearchMenu.cpp
--- a/Engine/Plugins/VirtualProduction/LevelSnapshots/Source/LevelSnapshotsEditor/Private/Widgets/Filter/SFilterSearchMenu.cpp
+++ b/Engine/Plugins/VirtualProduction/LevelSnapshots/Source/LevelSnapshotsEditor/Private/Widgets/Filter/SFilterSearchMenu.cpp
@@ -74,7 +74,7 @@ namespace
 				}
 
 				TArray<TSubclassOf<ULevelSnapshotFilter>> FiltersInCategory = FilterModel->GetFiltersInCategory(CategoryName);;
-				for (const TSubclassOf<ULevelSnapshotFilter> FilterClass : FiltersInCategory)
+				for (const TSubclassOf<ULevelSnapshotFilter>& FilterClass : FiltersInCategory)
 				{
 					CreateMenuEntry(SubmenuBuilder, FilterClass, OnSelectFilter, IsFilterSelected);
 				}
