Render Grid: fixed crash when passing in an empty string when setting remote control values

diff --git a/Engine/Plugins/Experimental/RenderGrid/Source/RenderGrid/Private/Utils/RenderGridUtils.cpp b/Engine/Plugins/Experimental/RenderGrid/Source/RenderGrid/Private/Utils/RenderGridUtils.cpp
--- a/Engine/Plugins/Experimental/RenderGrid/Source/RenderGrid/Private/Utils/RenderGridUtils.cpp
+++ b/Engine/Plugins/Experimental/RenderGrid/Source/RenderGrid/Private/Utils/RenderGridUtils.cpp
@@ -13,6 +13,7 @@
 #include "Misc/Paths.h"
 #include "Modules/ModuleManager.h"
 #include "Policies/CondensedJsonPrintPolicy.h"
+#include "RenderGridLog.h"
 #include "Serialization/JsonReader.h"
 #include "Serialization/JsonSerializer.h"
 
@@ -407,6 +408,11 @@ TArray<uint8> UE::RenderGrid::Private::FRenderGridUtils::GetRemoteControlValueBy
 	}
 
 	TSharedPtr<FJsonValue> Value = ParseJson(NewValueJson);
+	if (!Value.IsValid())
+	{
+		UE_LOG(LogRenderGrid, Warning, TEXT("Unable to retrieve remote control value from provided json. The json is invalid: '%s'"), *NewValueJson);
+		return TArray<uint8>();
+	}
 
 	FString JsonString;
 	{
